Add GameplayTargetingSystem plugin to engine

diff --git a/Engine/Plugins/Experimental/GameplayTargetingSystem/Source/GameplayTargetingSystem/Async/AsyncAction_PerformTargeting.cpp b/Engine/Plugins/Experimental/GameplayTargetingSystem/Source/GameplayTargetingSystem/Async/AsyncAction_PerformTargeting.cpp
new file mode 100644
--- /dev/null
+++ b/Engine/Plugins/Experimental/GameplayTargetingSystem/Source/GameplayTargetingSystem/Async/AsyncAction_PerformTargeting.cpp
@@ -0,0 +1,131 @@
+// Copyright Epic Games, Inc. All Rights Reserved.
+
+#include "AsyncAction_PerformTargeting.h"
+
+#include "Engine/Engine.h"
+#include "Engine/World.h"
+#include "GameFramework/Actor.h"
+#include "GameplayTargetingSystem/TargetingSystem/TargetingPreset.h"
+#include "GameplayTargetingSystem/TargetingSystem/TargetingSubsystem.h"
+#include "GameplayTargetingSystem/Types/TargetingSystemLogs.h"
+#include "GameplayTargetingSystem/Types/TargetingSystemTypes.h"
+
+
+UAsyncAction_PerformTargeting::UAsyncAction_PerformTargeting(const FObjectInitializer& ObjectInitializer)
+: Super(ObjectInitializer)
+{
+	bUseAsyncTargeting = false;
+}
+
+UAsyncAction_PerformTargeting* UAsyncAction_PerformTargeting::PerformTargetingRequest(AActor* SourceActor, UTargetingPreset* TargetingPreset, bool bUseAsyncTargeting)
+{
+	UWorld* World = GEngine->GetWorldFromContextObject(SourceActor, EGetWorldErrorMode::LogAndReturnNull);
+	if (!World)
+	{
+		return nullptr;
+	}
+
+	UAsyncAction_PerformTargeting* Action = NewObject<UAsyncAction_PerformTargeting>();
+	Action->TargetingPreset = TargetingPreset;
+	Action->WeakSourceActor = SourceActor;
+	Action->bUseAsyncTargeting = bUseAsyncTargeting;
+	Action->RegisterWithGameInstance(World);
+
+	return Action;
+}
+
+UAsyncAction_PerformTargeting* UAsyncAction_PerformTargeting::PerformFilteringRequest(AActor* SourceActor, UTargetingPreset* TargetingPreset, bool bUseAsyncTargeting, const TArray<AActor*> InTargets)
+{
+	UWorld* World = GEngine->GetWorldFromContextObject(SourceActor, EGetWorldErrorMode::LogAndReturnNull);
+	if (!World)
+	{
+		return nullptr;
+	}
+
+	UAsyncAction_PerformTargeting* Action = NewObject<UAsyncAction_PerformTargeting>();
+	Action->TargetingPreset = TargetingPreset;
+	Action->WeakSourceActor = SourceActor;
+	Action->bUseAsyncTargeting = bUseAsyncTargeting;
+	Action->InitialTargets = InTargets;
+	Action->RegisterWithGameInstance(World);
+
+	return Action;
+}
+
+void UAsyncAction_PerformTargeting::Activate()
+{
+	if (AActor* SourceActor = WeakSourceActor.Get())
+	{
+		if (UTargetingSubsystem* TargetingSubsystem = UTargetingSubsystem::Get(SourceActor->GetWorld()))
+		{
+			FTargetingSourceContext SourceContext;
+			SourceContext.SourceActor = SourceActor;
+
+			TARGETING_LOG(Verbose, TEXT("Source Actor: %s"), *GetNameSafe(SourceContext.SourceActor));
+
+			FTargetingRequestHandle TargetingHandle = UTargetingSubsystem::MakeTargetRequestHandle(TargetingPreset, SourceContext);
+
+			SetupInitialTargetsForRequest(TargetingHandle);
+
+			FTargetingRequestDelegate Delegate = FTargetingRequestDelegate::CreateWeakLambda(this, [this](FTargetingRequestHandle TargetingHandle)
+				{
+					TARGETING_LOG(Verbose, TEXT("Entering request lambda"));
+
+					Targeted.Broadcast(TargetingHandle);
+				});
+
+			if (bUseAsyncTargeting)
+			{
+				TARGETING_LOG(Verbose, TEXT("Starting async targeting"));
+
+				FTargetingAsyncTaskData& AsyncTaskData = FTargetingAsyncTaskData::FindOrAdd(TargetingHandle);
+				AsyncTaskData.bReleaseOnCompletion = true;
+
+				TargetingSubsystem->StartAsyncTargetingRequestWithHandle(TargetingHandle, Delegate);
+			}
+			else
+			{
+				TARGETING_LOG(Verbose, TEXT("Starting immediate targeting"));
+
+				FTargetingImmediateTaskData& ImmeidateTaskData = FTargetingImmediateTaskData::FindOrAdd(TargetingHandle);
+				ImmeidateTaskData.bReleaseOnCompletion = true;
+
+				TargetingSubsystem->ExecuteTargetingRequestWithHandle(TargetingHandle, Delegate);
+
+				SetReadyToDestroy();
+			}
+		}
+	}
+	else
+	{
+		SetReadyToDestroy();
+	}
+}
+
+void UAsyncAction_PerformTargeting::SetupInitialTargetsForRequest(const FTargetingRequestHandle& TargetingHandle) const
+{
+	if (TargetingHandle.IsValid() && InitialTargets.Num() > 0)
+	{
+		FTargetingDefaultResultsSet& TargetingResults = FTargetingDefaultResultsSet::FindOrAdd(TargetingHandle);
+		for (AActor* Target : InitialTargets)
+		{
+			if (!Target)
+			{
+				continue;
+			}
+
+			bool bAddResult = !TargetingResults.TargetResults.ContainsByPredicate([Target](const FTargetingDefaultResultData& Data) -> bool
+				{
+					return (Data.HitResult.GetActor() == Target);
+				});
+
+			if (bAddResult)
+			{
+				FTargetingDefaultResultData* ResultData = new(TargetingResults.TargetResults) FTargetingDefaultResultData();
+				ResultData->HitResult.HitObjectHandle = FActorInstanceHandle(Target);
+				ResultData->HitResult.Location = Target->GetActorLocation();
+			}
+		}
+
+	}
+}
diff --git a/Engine/Plugins/Experimental/GameplayTargetingSystem/Source/GameplayTargetingSystem/Async/AsyncAction_PerformTargeting.h b/Engine/Plugins/Experimental/GameplayTargetingSystem/Source/GameplayTargetingSystem/Async/AsyncAction_PerformTargeting.h
new file mode 100644
--- /dev/null
+++ b/Engine/Plugins/Experimental/GameplayTargetingSystem/Source/GameplayTargetingSystem/Async/AsyncAction_PerformTargeting.h
@@ -0,0 +1,59 @@
+// Copyright Epic Games, Inc. All Rights Reserved.
+
+#pragma once
+
+#include "CoreMinimal.h"
+#include "GameplayTargetingSystem/Types/TargetingSystemTypes.h"
+#include "Kismet/BlueprintAsyncActionBase.h"
+#include "UObject/ObjectMacros.h"
+
+#include "AsyncAction_PerformTargeting.generated.h"
+
+class AActor;
+class UTargetingPreset;
+
+DECLARE_DYNAMIC_MULTICAST_DELEGATE_OneParam(FPerformTargetingReady, FTargetingRequestHandle, TargetingHandle);
+
+
+/**
+*	@class UAsyncAction_PerformTargeting
+*/
+UCLASS()
+class UAsyncAction_PerformTargeting : public UBlueprintAsyncActionBase
+{
+	GENERATED_UCLASS_BODY()
+
+public:
+	UFUNCTION(BlueprintCallable, meta = (DefaultToSelf = "SourceActor", BlueprintInternalUseOnly = "true", DisplayName = "Perform Targeting Async Action"))
+	static UAsyncAction_PerformTargeting* PerformTargetingRequest(AActor* SourceActor, UTargetingPreset* TargetingPreset, bool bUseAsyncTargeting);
+
+	UFUNCTION(BlueprintCallable, meta = (DefaultToSelf = "SourceActor", BlueprintInternalUseOnly = "true", DisplayName = "Perform Filtering Async Action"))
+	static UAsyncAction_PerformTargeting* PerformFilteringRequest(AActor* SourceActor, UTargetingPreset* TargetingPreset, bool bUseAsyncTargeting, const TArray<AActor*> InTargets);
+
+	virtual void Activate() override;
+
+public:
+	UPROPERTY(BlueprintAssignable)
+	FPerformTargetingReady Targeted;
+
+private:
+	/** Method to seed the targeting request w/ the initial set of targets */
+	void SetupInitialTargetsForRequest(const FTargetingRequestHandle& TargetingHandle) const;
+
+private:
+	/** The targeting preset to use for targeting */
+	UPROPERTY()
+	TObjectPtr<UTargetingPreset> TargetingPreset = nullptr;
+
+	/** The actor this task is running for */
+	UPROPERTY()
+	TWeakObjectPtr<AActor> WeakSourceActor;
+
+	/** A set of targets to pre-seed the targeting request with (great for filtering targets without the needs for selection) */
+	UPROPERTY()
+	TArray<TObjectPtr<AActor>> InitialTargets;
+
+	/** Flag to indicate we should be using async targeting (default is immediate execution) */
+	UPROPERTY()
+	uint8 bUseAsyncTargeting : 1;
+};
diff --git a/Engine/Plugins/Experimental/GameplayTargetingSystem/Source/GameplayTargetingSystem/TargetingSystem.Build.cs b/Engine/Plugins/Experimental/GameplayTargetingSystem/Source/GameplayTargetingSystem/TargetingSystem.Build.cs
new file mode 100644
--- /dev/null
+++ b/Engine/Plugins/Experimental/GameplayTargetingSystem/Source/GameplayTargetingSystem/TargetingSystem.Build.cs
@@ -0,0 +1,40 @@
+// Copyright Epic Games, Inc. All Rights Reserved.
+using UnrealBuildTool;
+public class TargetingSystem : ModuleRules
+{
+	public TargetingSystem(ReadOnlyTargetRules Target) : base(Target)
+	{
+		PCHUsage = ModuleRules.PCHUsageMode.UseExplicitOrSharedPCHs;
+
+		PublicIncludePaths.AddRange(
+			new string[] {
+			}
+		);
+
+		PrivateIncludePaths.AddRange(
+			new string[] {
+            }
+		);
+
+		PublicDependencyModuleNames.AddRange(
+			new string[] {
+				"Core",
+				"CoreUObject",
+				"Engine",
+				"GameplayAbilities",
+				"GameplayTags",
+			}
+		);
+
+		PrivateDependencyModuleNames.AddRange(
+			new string[] {
+				"DeveloperSettings"
+			}
+		);
+
+		DynamicallyLoadedModuleNames.AddRange(
+			new string[] {
+			}
+		);
+	}
+}
\ No newline at end of file
diff --git a/Engine/Plugins/Experimental/GameplayTargetingSystem/Source/GameplayTargetingSystem/TargetingSystem/TargetingPreset.cpp b/Engine/Plugins/Experimental/GameplayTargetingSystem/Source/GameplayTargetingSystem/TargetingSystem/TargetingPreset.cpp
new file mode 100644
--- /dev/null
+++ b/Engine/Plugins/Experimental/GameplayTargetingSystem/Source/GameplayTargetingSystem/TargetingSystem/TargetingPreset.cpp
@@ -0,0 +1,9 @@
+// Copyright Epic Games, Inc. All Rights Reserved.
+#include "TargetingPreset.h"
+
+
+UTargetingPreset::UTargetingPreset(const FObjectInitializer& ObjectInitializer)
+{
+
+}
+
diff --git a/Engine/Plugins/Experimental/GameplayTargetingSystem/Source/GameplayTargetingSystem/TargetingSystem/TargetingPreset.h b/Engine/Plugins/Experimental/GameplayTargetingSystem/Source/GameplayTargetingSystem/TargetingSystem/TargetingPreset.h
new file mode 100644
--- /dev/null
+++ b/Engine/Plugins/Experimental/GameplayTargetingSystem/Source/GameplayTargetingSystem/TargetingSystem/TargetingPreset.h
@@ -0,0 +1,39 @@
+// Copyright Epic Games, Inc. All Rights Reserved.
+#pragma once
+
+#include "UObject/Object.h"
+#include "Engine/DataAsset.h"
+#include "GameplayTargetingSystem/Types/TargetingSystemTypes.h"
+
+#include "TargetingPreset.generated.h"
+
+class UTargetingTask;
+struct FTargetingRequestData;
+struct FTargetingDefaultResultsSet;
+
+
+/**
+*	@class UTargetingPreset
+*
+*	This object is used when a data asset is preferred for task setup.
+*
+*	Tasks will be processed in the order they are setup in the task set. It
+*	is recommended that selection tasks happen first before any filtering
+*	and sorting.
+*/
+UCLASS(BlueprintType)
+class UTargetingPreset : public UDataAsset
+{
+	GENERATED_BODY()
+
+public:
+	UTargetingPreset(const FObjectInitializer& ObjectInitializer);
+
+	/** Get the set of tasks defined for this targeting preset */
+	FORCEINLINE const FTargetingTaskSet* GetTargetingTaskSet() const { return &TargetingTaskSet; }
+
+private:
+	/** The tasks that make up this targeting preset */
+	UPROPERTY(EditAnywhere, Category = "Targeting", meta = (ShowOnlyInnerProperties))
+	FTargetingTaskSet TargetingTaskSet;
+};
\ No newline at end of file
diff --git a/Engine/Plugins/Experimental/GameplayTargetingSystem/Source/GameplayTargetingSystem/TargetingSystem/TargetingSubsystem.cpp b/Engine/Plugins/Experimental/GameplayTargetingSystem/Source/GameplayTargetingSystem/TargetingSystem/TargetingSubsystem.cpp
new file mode 100644
--- /dev/null
+++ b/Engine/Plugins/Experimental/GameplayTargetingSystem/Source/GameplayTargetingSystem/TargetingSystem/TargetingSubsystem.cpp
@@ -0,0 +1,858 @@
+// Copyright Epic Games, Inc. All Rights Reserved.
+#include "TargetingSubsystem.h"
+
+#include "Engine/World.h"
+#include "GameplayTargetingSystem/Types/TargetingSystemLogs.h"
+#include "GameplayTargetingSystem/Types/TargetingSystemTypes.h"
+#include "Kismet/KismetStringLibrary.h"
+#include "ProfilingDebugging/CsvProfiler.h"
+#include "Stats/Stats2.h"
+#include "TargetingPreset.h"
+
+#if ENABLE_DRAW_DEBUG
+#include "GameFramework/HUD.h"
+#include "Engine/Canvas.h"
+#include "DisplayDebugHelpers.h"
+#endif // ENABLE_DRAW_DEBUG
+
+namespace TargetingSystemCVars
+{
+	static bool bUseAsyncTargetingTimeSlicing = true;
+	FAutoConsoleVariableRef CvarUseAsyncTargetingTimeSlicing(
+		TEXT("ts.UseAsyncTargetingTimeSlicing"),
+		bUseAsyncTargetingTimeSlicing,
+		TEXT("Toggles whether the targeting system will use time slicing for the async targeting request queue. (Enabled: true, Disabled: false)")
+	);
+
+	static float MaxAsyncTickTime = .01f;
+	FAutoConsoleVariableRef CvarMaxAsyncTickTime(
+		TEXT("ts.MaxAsyncTickTime"),
+		MaxAsyncTickTime,
+		TEXT("Sets the total number of seconds we will allow to process the async targeting request queue.")
+	);
+
+#if ENABLE_DRAW_DEBUG
+
+	static bool bEnableTargetingDebugging = false;
+	FAutoConsoleVariableRef CvarEnableTargetingDebugging(
+		TEXT("ts.debug.EnableTargetingDebugging"),
+		bEnableTargetingDebugging,
+		TEXT("Toggles whether the targeting system is actively in debugging mode. (Enabled: true, Disabled: false)")
+	);
+
+	static bool bForceRequeueAsyncRequests = false;
+	FAutoConsoleVariableRef CvarForceRequeueAsyncRequests(
+		TEXT("ts.debug.ForceRequeueAsyncRequests"),
+		bForceRequeueAsyncRequests,
+		TEXT("Toggles whether the targeting system will force requeue async requests. (Enabled: true, Disabled: false)")
+	);
+
+	static bool bPrintTargetingDebugToLog = false;
+	FAutoConsoleVariableRef CvarPrintTargetingDebugToLog(
+		TEXT("ts.debug.PrintTargetingDebugToLog"),
+		bPrintTargetingDebugToLog,
+		TEXT("Toggles we print the targeting debug text to the log. (Enabled: true, Disabled: false)")
+	);
+
+	static int32 TotalDebugRecentRequestsTracked = 5;
+	FAutoConsoleVariableRef CvarTotalDebugRecentRequestsTracked(
+		TEXT("ts.debug.TotalDebugRecentRequestsTracked"),
+		TotalDebugRecentRequestsTracked,
+		TEXT("Sets the total # of targeting requests that will be tracked upon starting of them (default = 5)")
+	);
+
+#endif // ENABLE_DRAW_DEBUG
+}
+
+UTargetingSubsystem::FOnTargetingRequestHandleReleased UTargetingSubsystem::ReleaseHandleDelegate;
+
+
+UTargetingSubsystem::UTargetingSubsystem()
+: UGameInstanceSubsystem()
+{
+}
+
+/* static */
+UTargetingSubsystem* UTargetingSubsystem::Get(const UWorld* World)
+{
+	if (World)
+	{
+		return UGameInstance::GetSubsystem<UTargetingSubsystem>(World->GetGameInstance());
+	}
+
+	return nullptr;
+}
+
+/* static */
+UTargetingSubsystem* UTargetingSubsystem::GetTargetingSubsystem(const UObject* WorldContextObject)
+{
+	if (UWorld* World = GEngine->GetWorldFromContextObject(WorldContextObject, EGetWorldErrorMode::LogAndReturnNull))
+	{
+		return UTargetingSubsystem::Get(World);
+	}
+
+	return nullptr;
+}
+
+void UTargetingSubsystem::AddReferencedObjects(UObject* InThis, FReferenceCollector& Collector)
+{
+	Super::AddReferencedObjects(InThis, Collector);
+
+	UTargetingSubsystem* TypedThis = Cast<UTargetingSubsystem>(InThis);
+	check(TypedThis);
+
+	// Make sure any object references that async targeting requests are relying on are emitted
+	for (FTargetingRequestHandle& CurHandle : TypedThis->AsyncTargetingRequests)
+	{
+		if (FTargetingSourceContext* FoundSource = FTargetingSourceContext::Find(CurHandle))
+		{
+			Collector.AddReferencedObject(FoundSource->InstigatorActor);
+			Collector.AddReferencedObject(FoundSource->SourceActor);
+			Collector.AddReferencedObject(FoundSource->SourceObject);
+		}
+
+		if (FTargetingTaskSet** FoundTaskSetPtr = const_cast<FTargetingTaskSet**>(FTargetingTaskSet::Find(CurHandle)))
+		{
+			if (FTargetingTaskSet* FoundTaskSet = (*FoundTaskSetPtr))
+			{
+				Collector.AddReferencedObjects<UTargetingTask>(FoundTaskSet->Tasks);
+			}
+		}
+	}
+}
+
+void UTargetingSubsystem::Initialize(FSubsystemCollectionBase& Collection)
+{
+
+}
+
+void UTargetingSubsystem::Deinitialize()
+{
+	for (FTargetingRequestHandle AsyncHandle : AsyncTargetingRequests)
+	{
+		ReleaseTargetRequestHandle(AsyncHandle);
+	}
+
+	AsyncTargetingRequests.Empty();
+}
+
+bool UTargetingSubsystem::Exec(UWorld* Inworld, const TCHAR* Cmd, FOutputDevice& Ar)
+{
+	if (HasAnyFlags(RF_ClassDefaultObject) || IsDefaultSubobject())
+	{
+		return false;
+	}
+
+#if ENABLE_DRAW_DEBUG
+	if (FParse::Command(&Cmd, TEXT("ts.debug.ClearTrackedTargetRequests")))
+	{
+		CurrentImmediateRequestIndex = 0;
+		for (FTargetingRequestHandle Handle : DebugTrackedImmediateTargetRequests)
+		{
+			ReleaseTargetRequestHandle(Handle);
+		}
+
+		DebugTrackedImmediateTargetRequests.Empty();
+
+		for (FTargetingRequestHandle Handle : DebugTrackedAsyncTargetRequests)
+		{
+			ReleaseTargetRequestHandle(Handle);
+		}
+
+		DebugTrackedAsyncTargetRequests.Empty();
+
+		return true;
+	}
+#endif // ENABLE_DRAW_DEBUG
+
+	return false;
+}
+
+void UTargetingSubsystem::Tick(float DeltaTime)
+{
+	CSV_SCOPED_TIMING_STAT_EXCLUSIVE(TargetingSystem);
+	DECLARE_SCOPE_CYCLE_COUNTER(TEXT("UTargetingSubsystem::Tick"), STAT_TargetingSystem_Tick, STATGROUP_TargetingSystem);
+
+	// @note: Might need to implement a scheme to prevent bottlenecking on a single request
+	// or getting stuck from constant queuing of requests
+
+	float TimeLeft = TargetingSystemCVars::MaxAsyncTickTime;
+
+	const int32 NumRequests = AsyncTargetingRequests.Num();
+	for (int32 RequestIterator = 0; RequestIterator < NumRequests; ++RequestIterator)
+	{
+		const double StepStartTime = FPlatformTime::Seconds();
+
+		FTargetingRequestHandle& TargetingHandle = AsyncTargetingRequests[RequestIterator];
+		TARGETING_LOG(Verbose, TEXT("UTargetingSubsystem::Tick - Started Processing Async Request [%d] with Handle [%d]"), RequestIterator, TargetingHandle.Handle);
+
+		ProcessTargetingRequestTasks(TargetingHandle, TimeLeft);
+
+		if (TargetingSystemCVars::bUseAsyncTargetingTimeSlicing)
+		{
+			const double StepDuration = (FPlatformTime::Seconds() - StepStartTime);
+			TARGETING_LOG(VeryVerbose, TEXT("UTargetingSubsystem::Tick - Finished Processing Async Request [%d] in [%f] seconds."), RequestIterator, StepDuration);
+
+			TimeLeft -= StepDuration;
+			if (TimeLeft <= 0.0f)
+			{
+				TARGETING_LOG(VeryVerbose, TEXT("UTargetingSubsystem::Tick - Over alloted time, skipping remaining targeting requets this frame."));
+				break;
+			}
+		}
+	}
+	
+	{
+		DECLARE_SCOPE_CYCLE_COUNTER(TEXT("UTargetingSubsystem::Tick_HandleCleanup"), STAT_TargetingSystem_TickHandleCleanup, STATGROUP_TargetingSystem);
+
+		for (int32 RequestIterator = NumRequests - 1; RequestIterator >= 0; --RequestIterator)
+		{
+			FTargetingRequestHandle TargetingHandle = AsyncTargetingRequests[RequestIterator];
+			if (FTargetingRequestData* RequestData = FTargetingRequestData::Find(TargetingHandle))
+			{
+				if (RequestData->bComplete)
+				{
+					const bool bAllowShrinking = false;
+					AsyncTargetingRequests.RemoveAt(RequestIterator, 1, bAllowShrinking);
+
+					bool bForceRequeueOnCompletion = false;
+#if ENABLE_DRAW_DEBUG
+					bForceRequeueOnCompletion = IsTargetingDebugEnabled() && TargetingSystemCVars::bForceRequeueAsyncRequests && DebugTrackedAsyncTargetRequests.Contains(TargetingHandle);
+#endif // ENABLE_DRAW_DEBUG
+
+					if (FTargetingAsyncTaskData* AsyncTaskData = FTargetingAsyncTaskData::Find(TargetingHandle))
+					{
+						// re-queue if setup to do so
+						if (AsyncTaskData->bRequeueOnCompletion || bForceRequeueOnCompletion)
+						{
+							RequestData->bComplete = false;
+							AsyncTaskData->InitializeForAsyncProcessing();
+
+							if (FTargetingDefaultResultsSet* ResultSet = FTargetingDefaultResultsSet::Find(TargetingHandle))
+							{
+								ResultSet->TargetResults.Empty();
+							}
+
+							// put the handle back into the array at the end
+							AsyncTargetingRequests.Add(TargetingHandle);
+						}
+						// release the handle if setup to do so
+						else if (AsyncTaskData->bReleaseOnCompletion)
+						{
+							UTargetingSubsystem::ReleaseTargetRequestHandle(TargetingHandle);
+						}
+					}
+					else
+					{
+						TARGETING_LOG(Error, TEXT("[%s] - No Async Task Data for Handle [%d]"), ANSI_TO_TCHAR(__FUNCTION__), TargetingHandle.Handle)
+
+						UTargetingSubsystem::ReleaseTargetRequestHandle(TargetingHandle);
+					}
+				}
+			}
+		}
+	}
+}
+
+bool UTargetingSubsystem::IsTickable() const
+{
+	return !HasAnyFlags(RF_BeginDestroyed) && IsValidChecked(this) && (AsyncTargetingRequests.Num() > 0);
+}
+
+/* static */
+FTargetingRequestHandle UTargetingSubsystem::CreateTargetRequestHandle()
+{
+	static int32 HandleGenerator = 0;
+	FTargetingRequestHandle Handle(++HandleGenerator);
+	return Handle;
+}
+
+/* static */
+FTargetingRequestHandle UTargetingSubsystem::MakeTargetRequestHandle(const UTargetingPreset* TargetingPreset, const FTargetingSourceContext& InSourceContext)
+{
+	FTargetingRequestHandle Handle = CreateTargetRequestHandle();
+
+	// store the task set
+	if (TargetingPreset)
+	{
+		const FTargetingTaskSet*& TaskSet = FTargetingTaskSet::FindOrAdd(Handle);
+		TaskSet = TargetingPreset->GetTargetingTaskSet();
+
+		TARGETING_LOG(Verbose, TEXT("UTargetingSubsystem::MakeTargetRequestHandle - [%s] is Adding TaskSet [0x%08x] for Handle [%d]"), *GetNameSafe(TargetingPreset), TaskSet, Handle.Handle);
+	}
+	else
+	{
+		TARGETING_LOG(Warning, TEXT("UTargetingSubsystem::MakeTargetRequestHandle - NULL TargetingPreset, no tasks will be setup for the targeting handle [%d]"), Handle.Handle);
+	}
+
+	// store the source context
+	FTargetingSourceContext& SourceContext = FTargetingSourceContext::FindOrAdd(Handle);
+	SourceContext = InSourceContext;
+
+	return Handle;
+}
+
+/* static */
+void UTargetingSubsystem::ReleaseTargetRequestHandle(FTargetingRequestHandle& Handle)
+{
+	// Caching the handle to be able to broadcast it later after the reset
+	const FTargetingRequestHandle CachedHandle = Handle;
+	Handle.Reset();
+	
+#if ENABLE_DRAW_DEBUG
+	if (IsTargetingDebugEnabled())
+	{
+		return;
+	}
+#endif // ENABLE_DRAW_DEBUG
+
+	ReleaseHandleDelegate.Broadcast(CachedHandle);
+}
+
+void UTargetingSubsystem::ExecuteTargetingRequestWithHandle(FTargetingRequestHandle TargetingHandle, FTargetingRequestDelegate CompletionDelegate) const
+{
+	FTargetingRequestDynamicDelegate CompletionDynamicDelegate;
+	ExecuteTargetingRequestWithHandleInternal(TargetingHandle, CompletionDelegate, CompletionDynamicDelegate);
+
+	// if the caller setup task data to release on completion, lets do it
+	if (const FTargetingImmediateTaskData* ImmediateTaskData = FTargetingImmediateTaskData::Find(TargetingHandle))
+	{
+		if (ImmediateTaskData->bReleaseOnCompletion)
+		{
+			UTargetingSubsystem::ReleaseTargetRequestHandle(TargetingHandle);
+		}
+	}
+}
+
+void UTargetingSubsystem::ExecuteTargetingRequest(const UTargetingPreset* TargetingPreset, const FTargetingSourceContext& SourceContext, FTargetingRequestDynamicDelegate CompletionDynamicDelegate) const
+{
+	FTargetingRequestHandle RequestHandle;
+	if (ensure(TargetingPreset))
+	{
+		RequestHandle = MakeTargetRequestHandle(TargetingPreset, SourceContext);
+
+#if ENABLE_DRAW_DEBUG
+#if WITH_EDITORONLY_DATA
+		if (IsTargetingDebugEnabled())
+		{
+			FTargetingDebugData& DebugData = FTargetingDebugData::FindOrAdd(RequestHandle);
+			DebugData.TargetingPreset = TargetingPreset;
+		}
+#endif // WITH_EDITORONLY_DATA
+#endif // ENABLE_DRAW_DEBUG
+
+		FTargetingRequestDelegate CompletionDelegate;
+		ExecuteTargetingRequestWithHandleInternal(RequestHandle, CompletionDelegate, CompletionDynamicDelegate);
+		UTargetingSubsystem::ReleaseTargetRequestHandle(RequestHandle);
+	}
+}
+
+void UTargetingSubsystem::StartAsyncTargetingRequestWithHandle(FTargetingRequestHandle TargetingHandle, FTargetingRequestDelegate CompletionDelegate)
+{
+	FTargetingRequestDynamicDelegate CompletionDynamicDelegate;
+	StartAsyncTargetingRequestWithHandleInternal(TargetingHandle, CompletionDelegate, CompletionDynamicDelegate);
+}
+
+void UTargetingSubsystem::RemoveAsyncTargetingRequestWithHandle(FTargetingRequestHandle& TargetingHandle)
+{
+	if (TargetingHandle.IsValid())
+	{
+		AsyncTargetingRequests.RemoveAll([TargetingHandle](const FTargetingRequestHandle& Handle)
+			{
+				return (TargetingHandle == Handle);
+			});
+
+		UTargetingSubsystem::ReleaseTargetRequestHandle(TargetingHandle);
+	}
+}
+
+void UTargetingSubsystem::StartAsyncTargetingRequest(const UTargetingPreset* TargetingPreset, const FTargetingSourceContext& SourceContext, FTargetingRequestDynamicDelegate CompletionDynamicDelegate)
+{
+	FTargetingRequestHandle RequestHandle;
+	if (ensure(TargetingPreset))
+	{
+		RequestHandle = MakeTargetRequestHandle(TargetingPreset, SourceContext);
+
+#if ENABLE_DRAW_DEBUG
+#if WITH_EDITORONLY_DATA
+		if (IsTargetingDebugEnabled())
+		{
+			FTargetingDebugData& DebugData = FTargetingDebugData::FindOrAdd(RequestHandle);
+			DebugData.TargetingPreset = TargetingPreset;
+		}
+#endif // WITH_EDITORONLY_DATA
+#endif // ENABLE_DRAW_DEBUG
+
+		FTargetingRequestDelegate CompletionDelegate;
+		StartAsyncTargetingRequestWithHandleInternal(RequestHandle, CompletionDelegate, CompletionDynamicDelegate);
+	}
+}
+
+void UTargetingSubsystem::ExecuteTargetingRequestWithHandleInternal(FTargetingRequestHandle TargetingHandle, FTargetingRequestDelegate CompletionDelegate, FTargetingRequestDynamicDelegate CompletionDynamicDelegate) const
+{
+	if (ensure(TargetingHandle.IsValid()))
+	{
+		// store the request data
+		FTargetingRequestData& RequestData = FTargetingRequestData::FindOrAdd(TargetingHandle);
+		RequestData.Initialize(CompletionDelegate, CompletionDynamicDelegate);
+
+#if ENABLE_DRAW_DEBUG
+		if (IsTargetingDebugEnabled())
+		{
+			AddDebugTrackedImmediateTargetRequests(TargetingHandle);
+		}
+#endif // ENABLE_DRAW_DEBUG
+
+		if (const FTargetingTaskSet** PtrToTaskSet = FTargetingTaskSet::Find(TargetingHandle))
+		{
+			if (const FTargetingTaskSet* TaskSet = (*PtrToTaskSet))
+			{
+				RequestData.bComplete = false;
+
+				for (UTargetingTask* Task : TaskSet->Tasks)
+				{
+					if (Task)
+					{
+						Task->Init(TargetingHandle);
+						Task->Execute(TargetingHandle);
+					}
+				}
+
+				RequestData.bComplete = true;
+				RequestData.BroadcastTargetingRequestDelegate(TargetingHandle);
+
+#if ENABLE_DRAW_DEBUG
+#if WITH_EDITORONLY_DATA
+				if (const FTargetingDefaultResultsSet* ResultSet = FTargetingDefaultResultsSet::Find(TargetingHandle))
+				{
+					FTargetingDebugData& DebugData = FTargetingDebugData::FindOrAdd(TargetingHandle);
+					DebugData.CachedTargetResults = ResultSet->TargetResults;
+				}
+#endif // WITH_EDITORONLY_DATA
+#endif // ENABLE_DRAW_DEBUG
+			}
+		}
+	}
+}
+
+void UTargetingSubsystem::StartAsyncTargetingRequestWithHandleInternal(FTargetingRequestHandle TargetingHandle, FTargetingRequestDelegate CompletionDelegate, FTargetingRequestDynamicDelegate CompletionDynamicDelegate)
+{
+	// store the request data
+	FTargetingRequestData& RequestData = FTargetingRequestData::FindOrAdd(TargetingHandle);
+	RequestData.Initialize(CompletionDelegate, CompletionDynamicDelegate);
+
+	// initialize the request data for async processing
+	FTargetingAsyncTaskData& AsyncTaskData = FTargetingAsyncTaskData::FindOrAdd(TargetingHandle);
+	AsyncTaskData.InitializeForAsyncProcessing();
+
+#if ENABLE_DRAW_DEBUG
+	if (IsTargetingDebugEnabled())
+	{
+		AddDebugTrackedAsyncTargetRequests(TargetingHandle);
+	}
+#endif // ENABLE_DRAW_DEBUG
+
+	AsyncTargetingRequests.Add(TargetingHandle);
+}
+
+void UTargetingSubsystem::ProcessTargetingRequestTasks(FTargetingRequestHandle& TargetingHandle, float& TimeLeft)
+{
+	CSV_SCOPED_TIMING_STAT_EXCLUSIVE(TargetingSystem);
+	DECLARE_SCOPE_CYCLE_COUNTER(TEXT("UTargetingSubsystem::ProcessTargetingRequestTasks"), STAT_TargetingSystem_TickTasks, STATGROUP_TargetingSystem);
+
+	// @note: this function will need to add total time support for a task to mitigate tasks that stall and never complete.
+
+	if (FTargetingAsyncTaskData* AsyncStateData = FTargetingAsyncTaskData::Find(TargetingHandle))
+	{
+		if (const FTargetingTaskSet** PtrToTaskSet = FTargetingTaskSet::Find(TargetingHandle))
+		{
+			TARGETING_LOG(Verbose, TEXT("UTargetingSubsystem::ProcessTargetingRequestTasks - Processing TaskSet [0x%08x] for Handle [%d]"), (*PtrToTaskSet), TargetingHandle.Handle);
+			if (const FTargetingTaskSet* TaskSet = (*PtrToTaskSet))
+			{
+				bool bRequestComplete = false;
+
+				int32 CurrentTaskIndex = AsyncStateData->CurrentAsyncTaskIndex;
+				const int32 NumTasks = TaskSet->Tasks.Num();
+				while (CurrentTaskIndex < NumTasks)
+				{
+					const double StepStartTime = FPlatformTime::Seconds();
+
+					TARGETING_LOG(VeryVerbose, TEXT("UTargetingSubsystem::ProcessTargetingRequestTasks - Processing Async Index [%d] of [%d] for Handle [%d]"), CurrentTaskIndex, NumTasks, TargetingHandle.Handle);
+
+					if (UTargetingTask* Task = TaskSet->Tasks[CurrentTaskIndex])
+					{
+						TARGETING_LOG(VeryVerbose, TEXT("UTargetingSubsystem::ProcessTargetingRequestTasks - Started Processing Task [%s]"), *(GetNameSafe(Task)));
+
+						if (AsyncStateData->CurrentAsyncTaskState == ETargetingTaskAsyncState::Unitialized)
+						{
+							Task->Init(TargetingHandle);
+						}
+
+						if (AsyncStateData->CurrentAsyncTaskState == ETargetingTaskAsyncState::Initialized)
+						{
+							Task->Execute(TargetingHandle);
+						}
+
+						// @note: if the task is still executing (it has async support internally), we are done
+						// until it is done
+						if (AsyncStateData->CurrentAsyncTaskState == ETargetingTaskAsyncState::Executing)
+						{
+							TARGETING_LOG(VeryVerbose, TEXT("UTargetingSubsystem::ProcessTargetingRequestTasks - Waiting on Task [%s] to finish."), *(GetNameSafe(Task)));
+							break;
+						}
+
+						// if we got here then the task should be complete
+						ensure(AsyncStateData->CurrentAsyncTaskState == ETargetingTaskAsyncState::Completed);
+						AsyncStateData->CurrentAsyncTaskState = ETargetingTaskAsyncState::Unitialized;
+
+						++CurrentTaskIndex;
+						if (CurrentTaskIndex >= NumTasks)
+						{
+							bRequestComplete = true;
+						}
+						else
+						{
+							AsyncStateData->CurrentAsyncTaskIndex = CurrentTaskIndex;
+
+							// if we are out of time, save the task for the next time
+							if (TargetingSystemCVars::bUseAsyncTargetingTimeSlicing)
+							{
+								const double StepDuration = (FPlatformTime::Seconds() - StepStartTime);
+								TARGETING_LOG(VeryVerbose, TEXT("UTargetingSubsystem::ProcessTargetingRequestTasks - Finished Processing Task [%s] in [%f] seconds."), *(GetNameSafe(Task)), StepDuration);
+
+								TimeLeft -= StepDuration;
+								if (TimeLeft <= 0.0f)
+								{
+									break;
+								}
+							}
+						}
+					}
+				}
+
+				if (bRequestComplete)
+				{
+					if (FTargetingRequestData* RequestData = FTargetingRequestData::Find(TargetingHandle))
+					{
+						RequestData->bComplete = true;
+						RequestData->BroadcastTargetingRequestDelegate(TargetingHandle);
+
+#if ENABLE_DRAW_DEBUG
+#if WITH_EDITORONLY_DATA
+						if (const FTargetingDefaultResultsSet* ResultSet = FTargetingDefaultResultsSet::Find(TargetingHandle))
+						{
+							FTargetingDebugData& DebugData = FTargetingDebugData::FindOrAdd(TargetingHandle);
+							DebugData.CachedTargetResults = ResultSet->TargetResults;
+						}
+#endif // WITH_EDITORONLY_DATA
+#endif // ENABLE_DRAW_DEBUG
+					}
+				}
+			}
+		}
+	}
+}
+
+FTargetingSourceContext UTargetingSubsystem::GetTargetingSourceContext(FTargetingRequestHandle TargetingHandle)
+{
+	if (TargetingHandle.IsValid())
+	{
+		FTargetingSourceContext& SourceContext = FTargetingSourceContext::FindOrAdd(TargetingHandle);
+		return SourceContext;
+	}
+
+	return FTargetingSourceContext();
+}
+
+/* static */
+void UTargetingSubsystem::GetTargetingResultsActors(FTargetingRequestHandle TargetingHandle, TArray<AActor*>& Targets)
+{
+	if (TargetingHandle.IsValid())
+	{
+		if (FTargetingDefaultResultsSet* Results = FTargetingDefaultResultsSet::Find(TargetingHandle))
+		{
+			for (const FTargetingDefaultResultData& ResultData : Results->TargetResults)
+			{
+				if (AActor* Target = ResultData.HitResult.GetActor())
+				{
+					Targets.Add(Target);
+				}
+			}
+		}
+	}
+}
+
+/* static */
+void UTargetingSubsystem::GetTargetingResults(FTargetingRequestHandle TargetingHandle, TArray<FHitResult>& OutTargets)
+{
+	if (TargetingHandle.IsValid())
+	{
+		if (FTargetingDefaultResultsSet* Results = FTargetingDefaultResultsSet::Find(TargetingHandle))
+		{
+			for (const FTargetingDefaultResultData& ResultData : Results->TargetResults)
+			{
+				OutTargets.Add(ResultData.HitResult);
+			}
+		}
+	}
+}
+
+#if ENABLE_DRAW_DEBUG
+
+bool UTargetingSubsystem::IsTargetingDebugEnabled()
+{
+	return TargetingSystemCVars::bEnableTargetingDebugging;
+}
+
+/* static */
+void UTargetingSubsystem::OnShowDebugInfo(AHUD* HUD, UCanvas* Canvas, const FDebugDisplayInfo& DisplayInfo, float& YL, float& YPos)
+{
+	if (IsTargetingDebugEnabled())
+	{
+		if (DisplayInfo.IsDisplayOn(TEXT("TargetingSystem")))
+		{
+			if (UWorld* World = HUD->GetWorld())
+			{
+				if (UTargetingSubsystem* TargetingSubsystem = UGameInstance::GetSubsystem<UTargetingSubsystem>(World->GetGameInstance()))
+				{
+					TargetingSubsystem->DisplayDebug(Canvas, DisplayInfo, YL, YPos);
+				}
+			}
+		}
+	}
+}
+
+void UTargetingSubsystem::DisplayDebug(class UCanvas* Canvas, const class FDebugDisplayInfo& DebugDisplay, float& YL, float& YPos)
+{
+	FTargetingDebugInfo DebugInfo;
+
+	DebugInfo.bPrintToLog = TargetingSystemCVars::bPrintTargetingDebugToLog;
+	DebugInfo.Canvas = Canvas;
+	DebugInfo.XPos = 0.f;
+	DebugInfo.YPos = YPos;
+	DebugInfo.OriginalX = 0.f;
+	DebugInfo.OriginalY = YPos;
+	DebugInfo.MaxY = Canvas->ClipY - 150.f; // Give some padding for any non-columnizing debug output following this output
+	DebugInfo.NewColumnYPadding = 30.f;
+
+	Debug_Internal(DebugInfo);
+
+	YPos = DebugInfo.YPos;
+	YL = DebugInfo.YL;
+}
+
+void UTargetingSubsystem::Debug_Internal(struct FTargetingDebugInfo& Info)
+{
+	FString DebugTitle("Targeting System Debug Info");
+
+	if (Info.Canvas)
+	{
+		Info.Canvas->SetDrawColor(FColor::White);
+		FFontRenderInfo RenderInfo = FFontRenderInfo();
+		RenderInfo.bEnableShadow = true;
+		Info.Canvas->DrawText(GEngine->GetLargeFont(), DebugTitle, Info.XPos + 4.f, 10.f, 1.5f, 1.5f, RenderInfo);
+	}
+	else
+	{
+		DebugLine(Info, DebugTitle, 0.f, 0.f);
+	}
+
+	float MaxCharHeight = 10;
+
+	DebugLine(Info, TEXT(""), 0.f, 0.f, 1);
+	DebugLine(Info, TEXT("Immediate Targeting Requests:"), 4.0f, 0.0f);
+
+	for (FTargetingRequestHandle& Handle : DebugTrackedImmediateTargetRequests)
+	{
+		DebugForHandle_Internal(Handle, Info);
+	}
+
+	DebugLine(Info, TEXT(""), 0.f, 0.f, 1);
+	DebugLine(Info, TEXT("Async Targeting Requests:"), 4.0f, 0.0f);
+
+	for (FTargetingRequestHandle& Handle : DebugTrackedAsyncTargetRequests)
+	{
+		DebugForHandle_Internal(Handle, Info);
+	}
+
+	if (Info.XPos > Info.OriginalX)
+	{
+		// We flooded to new columns, returned YPos should be max Y (and some padding)
+		Info.YPos = Info.MaxY + MaxCharHeight * 2.f;
+	}
+	Info.YL = MaxCharHeight;
+}
+
+void UTargetingSubsystem::DebugForHandle_Internal(FTargetingRequestHandle& Handle, struct FTargetingDebugInfo& Info)
+{
+#if WITH_EDITORONLY_DATA
+	if (Info.Canvas)
+	{
+		Info.Canvas->SetDrawColor(FColor::White);
+	}
+
+	if (Handle.IsValid())
+	{
+		const FTargetingDebugData* DebugData = FTargetingDebugData::Find(Handle);
+
+		FString HandleString = FString::Printf(TEXT("Targeting Handle [%d]"), Handle.Handle);
+		DebugLine(Info, HandleString, 8.0f, 0.0f);
+
+		// dump preset
+		FString PresetName = TEXT("Targeting Preset: <none>");
+		const UTargetingPreset* RequestPreset = nullptr;
+		if (DebugData)
+		{
+			RequestPreset = DebugData->TargetingPreset;
+			PresetName = FString::Printf(TEXT("Targeting Preset: %s"), *GetNameSafe(RequestPreset));
+		}
+		DebugLine(Info, PresetName, 8.0f, 0.0f);
+
+		// dump tasks
+		FString TasksString = TEXT("Tasks:");
+		DebugLine(Info, TasksString, 8.0f, 0.0f);
+		if (const FTargetingTaskSet** TaskSetPtr = FTargetingTaskSet::Find(Handle))
+		{
+			if (const FTargetingTaskSet* TaskSet = (*TaskSetPtr))
+			{
+				for (const UTargetingTask* Task : TaskSet->Tasks)
+				{
+					if (Task)
+					{
+						FString TaskString;
+						TaskString = FString::Printf(TEXT("%s"), *GetNameSafe(Task));
+
+						if (Info.Canvas)
+						{
+							Info.Canvas->SetDrawColor(FColor::White);
+						}
+						DebugLine(Info, TaskString, 12.0f, 0.0f);
+
+						Task->DrawDebug(this, Info, Handle, 16.0f, 0.0f);
+					}
+				}
+			}
+		}
+
+		// dump results
+		if (Info.Canvas)
+		{
+			Info.Canvas->SetDrawColor(FColor::White);
+		}
+
+		DebugLine(Info, TEXT("Targets:"), 8.0f, 0.0f);
+		if (DebugData && DebugData->CachedTargetResults.Num() > 0)
+		{
+			if (Info.Canvas)
+			{
+				Info.Canvas->SetDrawColor(FColor::Green);
+			}
+
+			FString TargetsString;
+			for (const FTargetingDefaultResultData& Result : DebugData->CachedTargetResults)
+			{
+				const AActor* TargetActor = Result.HitResult.GetActor();
+				if (TargetActor)
+				{
+					const float Score = Result.Score;
+
+					TargetsString = FString::Printf(TEXT("%s (%f)"), *GetNameSafe(TargetActor), Score);
+					DebugLine(Info, TargetsString, 12.0f, 0.0f);
+				}
+			}
+		}
+
+		if (Info.Canvas)
+		{
+			Info.Canvas->SetDrawColor(FColor::White);
+		}
+		DebugLine(Info, TEXT(""), 0.0f, 0.0f, 2);
+	}
+#endif // WITH_EDITORONLY_DATA
+}
+
+void UTargetingSubsystem::AccumulateScreenPos(FTargetingDebugInfo& Info)
+{
+	const float ColumnWidth = Info.Canvas ? Info.Canvas->ClipX * 0.4f : 0.f;
+
+	float NewY = Info.YPos + Info.YL;
+	if (NewY > Info.MaxY)
+	{
+		// Need new column, reset Y to original height
+		NewY = Info.NewColumnYPadding;
+		Info.XPos += ColumnWidth;
+	}
+	Info.YPos = NewY;
+}
+
+void UTargetingSubsystem::DebugLine(struct FTargetingDebugInfo& Info, FString Str, float XOffset, float YOffset, int32 MinTextRowsToAdvance /*= 0*/)
+{
+	if (Info.Canvas)
+	{
+		FFontRenderInfo RenderInfo = FFontRenderInfo();
+		RenderInfo.bEnableShadow = true;
+		if (const UFont* Font = GEngine->GetTinyFont())
+		{
+			float ScaleY = 1.f;
+			Info.YL = Info.Canvas->DrawText(Font, Str, Info.XPos + XOffset, Info.YPos, 1.f, ScaleY, RenderInfo);
+			if (Info.YL < MinTextRowsToAdvance * (Font->GetMaxCharHeight() * ScaleY))
+			{
+				Info.YL = MinTextRowsToAdvance * (Font->GetMaxCharHeight() * ScaleY);
+			}
+			AccumulateScreenPos(Info);
+		}
+	}
+
+	if (Info.bPrintToLog)
+	{
+		FString LogStr;
+		for (int32 i = 0; i < (int32)XOffset; ++i)
+		{
+			LogStr += TEXT(" ");
+		}
+		LogStr += Str;
+		TARGETING_LOG(Warning, TEXT("%s"), *LogStr);
+	}
+}
+
+void UTargetingSubsystem::AddDebugTrackedImmediateTargetRequests(FTargetingRequestHandle TargetingHandle) const
+{
+	int32 Index = CurrentImmediateRequestIndex++;
+	if (Index >= TargetingSystemCVars::TotalDebugRecentRequestsTracked)
+	{
+		Index = 0;
+		CurrentImmediateRequestIndex = 0;
+	}
+
+	if (!DebugTrackedImmediateTargetRequests.IsValidIndex(Index))
+	{
+		DebugTrackedImmediateTargetRequests.SetNumZeroed(TargetingSystemCVars::TotalDebugRecentRequestsTracked);
+	}
+
+	FTargetingRequestHandle PreviousHandle = DebugTrackedImmediateTargetRequests[Index];
+	ReleaseTargetRequestHandle(PreviousHandle);
+
+	DebugTrackedImmediateTargetRequests[Index] = TargetingHandle;
+}
+
+void UTargetingSubsystem::AddDebugTrackedAsyncTargetRequests(FTargetingRequestHandle TargetingHandle) const
+{
+	int32 Index = CurrentAsyncRequestIndex++;
+	if (Index >= TargetingSystemCVars::TotalDebugRecentRequestsTracked)
+	{
+		Index = 0;
+		CurrentAsyncRequestIndex = 0;
+	}
+
+	if (!DebugTrackedAsyncTargetRequests.IsValidIndex(Index))
+	{
+		DebugTrackedAsyncTargetRequests.SetNumZeroed(TargetingSystemCVars::TotalDebugRecentRequestsTracked);
+	}
+
+	FTargetingRequestHandle PreviousHandle = DebugTrackedAsyncTargetRequests[Index];
+	ReleaseTargetRequestHandle(PreviousHandle);
+
+	DebugTrackedAsyncTargetRequests[Index] = TargetingHandle;
+}
+
+#endif // ENABLE_DRAW_DEBUG
\ No newline at end of file
diff --git a/Engine/Plugins/Experimental/GameplayTargetingSystem/Source/GameplayTargetingSystem/TargetingSystem/TargetingSubsystem.h b/Engine/Plugins/Experimental/GameplayTargetingSystem/Source/GameplayTargetingSystem/TargetingSystem/TargetingSubsystem.h
new file mode 100644
--- /dev/null
+++ b/Engine/Plugins/Experimental/GameplayTargetingSystem/Source/GameplayTargetingSystem/TargetingSystem/TargetingSubsystem.h
@@ -0,0 +1,212 @@
+// Copyright Epic Games, Inc. All Rights Reserved.
+#pragma once
+
+#include "CoreMinimal.h"
+#include "DrawDebugHelpers.h"
+#include "GameplayTargetingSystem/Tasks/TargetingTask.h"
+#include "GameplayTargetingSystem/Types/TargetingSystemTypes.h"
+#include "Misc/CoreMisc.h"
+#include "Subsystems/GameInstanceSubsystem.h"
+#include "Tickable.h"
+
+#include "TargetingSubsystem.generated.h"
+
+class AHUD;
+class FDebugDisplayInfo;
+class FSubsystemCollectionBase;
+class UCanvas;
+class UTargetingPreset;
+class UTargetingTask;
+class UWorld;
+
+
+#if ENABLE_DRAW_DEBUG
+/**
+*	@struct FTargetingDebugInfo
+*/
+struct FTargetingDebugInfo
+{
+	FTargetingDebugInfo()
+	{
+		FMemory::Memzero(*this);
+	}
+
+	class UCanvas* Canvas;
+
+	bool bPrintToLog;
+
+	float XPos;
+	float YPos;
+	float OriginalX;
+	float OriginalY;
+	float MaxY;
+	float NewColumnYPadding;
+	float YL;
+};
+#endif // ENABLE_DRAW_DEBUG
+
+/**
+*	@class UTargetingSubsystem
+*
+*	The Targeting Subsystem is the entry point for users to initiate targeting requests.
+*
+*	The entry point to the system is the target request handle. The handle is used to interface with
+*	the targeting data stores. Data stores are templated classes around generic data structs that the
+*	system and tasks use to accomplish a targeting request.
+*
+*	The targeting system has 3 mandatory data stores and 1 required for async targeting request. These
+*	data stores are required to be set up before the system can properly run a targeting request. The 
+*	mandatory 3 data stores are FTargetingRequestData, FTargetingTaskSet, and FTargetingSourceContext.
+*	FTargetingAsyncTaskData is implicitly setup when an async targeting request is initiated.
+*
+*	Users can do all the pieces manually in C++ by setting up the required data stores themselves, 
+*	or, to have it a bit more automated, the user can use the APIs that utilize UTargetingPreset data asset. 
+*
+*	For immediate targeting requests users will call the Execute methods. These functions perform all the
+*	tasks till completion. The system will not go latent.
+*
+*	For async targeting requests users will call the Start Async methods. The system will queue up a targeting
+*	request and as each task is processed the system can run through all the tasks to completion or stop processing
+*	until the next frame while it waits for a task to complete.
+* 
+*	Note about Targeting Handles, when a targeting handle is created it will not implicitly release the handle.
+*	It is up to the creator to either grab a Async Task Data or Immediate Task Data and set a flag indicating 
+*	the system should do it for them after the callback fires, or it is up to the user to release the handle
+*	when they are done with it.
+*/
+UCLASS(MinimalAPI, DisplayName = "Targeting Subsystem")
+class UTargetingSubsystem : public UGameInstanceSubsystem, public FTickableGameObject, public FSelfRegisteringExec
+{
+	GENERATED_BODY()
+
+public:
+	/** Method to get the mesh network events subsystem */
+	TARGETINGSYSTEM_API static UTargetingSubsystem* Get(const UWorld* World);
+	TARGETINGSYSTEM_API static UTargetingSubsystem* GetTargetingSubsystem(const UObject* WorldContextObject);
+
+	/** Implemented to emit references from the global data stores */
+	static void AddReferencedObjects(UObject* InThis, FReferenceCollector& Collector);
+
+private:
+	UTargetingSubsystem();
+
+	// Begin USubsystem
+	virtual void Initialize(FSubsystemCollectionBase& Collection) override;
+	virtual void Deinitialize() override;
+	// End USubsystem
+
+	/** FSelfRegisteringExec implementation */
+	virtual bool Exec(UWorld* Inworld, const TCHAR* Cmd, FOutputDevice& Ar) override;
+	/** ~FSelfRegisteringExec implementation */
+
+	// FTickableGameObject interface
+	virtual void Tick(float DeltaTime) override;
+	virtual bool IsTickable() const override;
+	virtual ETickableTickType GetTickableTickType() const override { return (HasAnyFlags(RF_ClassDefaultObject) ? ETickableTickType::Never : ETickableTickType::Conditional); }
+	virtual TStatId GetStatId() const override { RETURN_QUICK_DECLARE_CYCLE_STAT(UTargetingSubsystem, STATGROUP_Tickables); }
+	virtual UWorld* GetTickableGameObjectWorld() const override { return GetWorld(); }
+	// ~FTickableGameObject interface
+
+public:
+	/** Target Handle Generation Methods */
+
+	/** Method to create a basic target request handle (nothing will be setup for any targeting requests)
+		The request data store, task set data store, and source context data store will have to be setup
+		manually by caller. */
+	TARGETINGSYSTEM_API static FTargetingRequestHandle CreateTargetRequestHandle();
+
+	/** Method to create and setup a target request handle given a targeting preset and source context */
+	TARGETINGSYSTEM_API static FTargetingRequestHandle MakeTargetRequestHandle(const UTargetingPreset* TargetingPreset, const FTargetingSourceContext& InSourceContext);
+
+	/** Method to release a target handle (Note: Unless flagged via AsyncTaskData or ImmediateTaskData this will not release implicitly) */
+	TARGETINGSYSTEM_API static void ReleaseTargetRequestHandle(FTargetingRequestHandle& Handle);
+
+	/** The handle released delegate that fires right before a handle is release so all data stores can clean up their state */
+	DECLARE_MULTICAST_DELEGATE_OneParam(FOnTargetingRequestHandleReleased, FTargetingRequestHandle)
+	static FOnTargetingRequestHandleReleased ReleaseHandleDelegate;
+
+	/** Target Handle Generation Methods */
+
+	/** Targeting Request Methods */
+public:
+	/** Method to execute an immediate targeting request with a given targeting handle. */
+	TARGETINGSYSTEM_API void ExecuteTargetingRequestWithHandle(FTargetingRequestHandle TargetingHandle, FTargetingRequestDelegate CompletionDelegate) const;
+
+	/** Method to execute an immediate targeting request based on a gameplay targeting preset.*/
+	UFUNCTION(BlueprintCallable, BlueprintPure = false, Category = "Gameplay Targeting | Instant Request")
+	void ExecuteTargetingRequest(const UTargetingPreset* TargetingPreset, const FTargetingSourceContext& InSourceContext, FTargetingRequestDynamicDelegate CompletionDynamicDelegate) const;
+
+	/** Method to queue an async targeting request with a given targeting handle. */
+	TARGETINGSYSTEM_API void StartAsyncTargetingRequestWithHandle(FTargetingRequestHandle TargetingHandle, FTargetingRequestDelegate CompletionDelegate);
+
+	/** Method to remove an async targeting request with a given targeting handle */
+	TARGETINGSYSTEM_API void RemoveAsyncTargetingRequestWithHandle(FTargetingRequestHandle& TargetingHandle);
+
+	/** Method to queue an async targeting request based on a gameplay targeting preset. */
+	UFUNCTION(BlueprintCallable, BlueprintPure = false, Category = "Gameplay Targeting | Async Request ")
+	void StartAsyncTargetingRequest(const UTargetingPreset* TargetingPreset, const FTargetingSourceContext& InSourceContext, FTargetingRequestDynamicDelegate CompletionDynamicDelegate);
+
+private:
+	/** Internal method to handle running an immediate targeting request */
+	void ExecuteTargetingRequestWithHandleInternal(FTargetingRequestHandle TargetingHandle, FTargetingRequestDelegate CompletionDelegate, FTargetingRequestDynamicDelegate CompletionDynamicDelegate) const;
+
+	/** Internal method to hanlde queuing up an async targeting request */
+	void StartAsyncTargetingRequestWithHandleInternal(FTargetingRequestHandle TargetingHandle, FTargetingRequestDelegate CompletionDelegate, FTargetingRequestDynamicDelegate CompletionDynamicDelegate);
+
+	/** Internal method to process the tasks for a given targeting request */
+	void ProcessTargetingRequestTasks(FTargetingRequestHandle& TargetingHandle, float& TimeLeft);
+
+	/** The set of target requests queued up for async processing */
+	UPROPERTY(Transient)
+	TArray<FTargetingRequestHandle> AsyncTargetingRequests;
+
+	/** ~Targeting Request Methods */
+
+	/** Blueprint Helper Methods */
+public:
+	/** Returns the targeting source context for the targeting request handle */
+	UFUNCTION(BlueprintCallable, Category = "Targeting System | Targeting Types")
+	static FTargetingSourceContext GetTargetingSourceContext(FTargetingRequestHandle TargetingHandle);
+
+	/** Method to get the actor targets from a given targeting request handle */
+	UFUNCTION(BlueprintCallable, Category = "Targeting System | Targeting Results")
+	static void GetTargetingResultsActors(FTargetingRequestHandle TargetingHandle, TArray<AActor*>& Targets);
+
+	/** Helper method to get the set of hit results for a given targeting handle */
+	UFUNCTION(BlueprintCallable, Category = "Targeting System | Targeting Results")
+	static TARGETINGSYSTEM_API void GetTargetingResults(FTargetingRequestHandle TargetingHandle, TArray<FHitResult>& OutTargets);
+
+	/** ~Blueprint Helper Methods */
+
+	/** Debug Helper Methods */
+#if ENABLE_DRAW_DEBUG
+
+public:
+	/** Helper method to indicate if the debugging feature is enabled */
+	TARGETINGSYSTEM_API static bool IsTargetingDebugEnabled();
+
+	/** Callback to handle OnShowDebugInfo to check if the TargetingSystem has been enabled */
+	static void OnShowDebugInfo(AHUD* HUD, UCanvas* Canvas, const FDebugDisplayInfo& DisplayInfo, float& YL, float& YPos);
+
+	/** Helper methd to write a debug line when the ShowDebug TargetingSystem is active */
+	TARGETINGSYSTEM_API void DebugLine(struct FTargetingDebugInfo& Info, FString Str, float XOffset, float YOffset, int32 MinTextRowsToAdvance = 0);
+
+private:
+	void DisplayDebug(class UCanvas* Canvas, const class FDebugDisplayInfo& DebugDisplay, float& YL, float& YPos);
+	void Debug_Internal(struct FTargetingDebugInfo& Info);
+	void DebugForHandle_Internal(FTargetingRequestHandle& Handle, struct FTargetingDebugInfo& Info);
+	void AccumulateScreenPos(FTargetingDebugInfo& Info);
+
+	void AddDebugTrackedImmediateTargetRequests(FTargetingRequestHandle TargetingHandle) const;
+	void AddDebugTrackedAsyncTargetRequests(FTargetingRequestHandle TargetingHandle) const;
+
+	mutable TArray<FTargetingRequestHandle> DebugTrackedImmediateTargetRequests;
+	mutable int32 CurrentImmediateRequestIndex = 0;
+
+	mutable TArray<FTargetingRequestHandle> DebugTrackedAsyncTargetRequests;
+	mutable int32 CurrentAsyncRequestIndex = 0;
+
+#endif // ENABLE_DRAW_DEBUG
+	/** ~Debug Helper Methods */
+};
+
diff --git a/Engine/Plugins/Experimental/GameplayTargetingSystem/Source/GameplayTargetingSystem/TargetingSystemModule.cpp b/Engine/Plugins/Experimental/GameplayTargetingSystem/Source/GameplayTargetingSystem/TargetingSystemModule.cpp
new file mode 100644
--- /dev/null
+++ b/Engine/Plugins/Experimental/GameplayTargetingSystem/Source/GameplayTargetingSystem/TargetingSystemModule.cpp
@@ -0,0 +1,27 @@
+// Copyright Epic Games, Inc. All Rights Reserved.
+#include "TargetingSystemModule.h"
+
+#include "TargetingSystem/TargetingSubsystem.h"
+#include "Types/TargetingSystemLogs.h"
+
+#if ENABLE_DRAW_DEBUG
+#include "GameFramework/HUD.h"
+#endif // ENABLE_DRAW_DEBUG
+
+IMPLEMENT_MODULE(FTargetingSystemModule, TargetingSystem)
+DEFINE_LOG_CATEGORY(LogTargetingSystem);
+
+
+void FTargetingSystemModule::StartupModule()
+{
+#if ENABLE_DRAW_DEBUG
+	if (!IsRunningDedicatedServer())
+	{
+		AHUD::OnShowDebugInfo.AddStatic(&UTargetingSubsystem::OnShowDebugInfo);
+	}
+#endif // ENABLE_DRAW_DEBUG
+}
+
+void FTargetingSystemModule::ShutdownModule()
+{
+}
diff --git a/Engine/Plugins/Experimental/GameplayTargetingSystem/Source/GameplayTargetingSystem/TargetingSystemModule.h b/Engine/Plugins/Experimental/GameplayTargetingSystem/Source/GameplayTargetingSystem/TargetingSystemModule.h
new file mode 100644
--- /dev/null
+++ b/Engine/Plugins/Experimental/GameplayTargetingSystem/Source/GameplayTargetingSystem/TargetingSystemModule.h
@@ -0,0 +1,16 @@
+// Copyright Epic Games, Inc. All Rights Reserved.
+
+#pragma once
+
+#include "CoreMinimal.h"
+#include "Stats/Stats.h"
+#include "Modules/ModuleManager.h"
+
+
+class FTargetingSystemModule : public IModuleInterface
+{
+public:
+
+	virtual void StartupModule() override;
+	virtual void ShutdownModule() override;
+};
\ No newline at end of file
diff --git a/Engine/Plugins/Experimental/GameplayTargetingSystem/Source/GameplayTargetingSystem/Tasks/TargetingFilterTask_ActorClass.cpp b/Engine/Plugins/Experimental/GameplayTargetingSystem/Source/GameplayTargetingSystem/Tasks/TargetingFilterTask_ActorClass.cpp
new file mode 100644
--- /dev/null
+++ b/Engine/Plugins/Experimental/GameplayTargetingSystem/Source/GameplayTargetingSystem/Tasks/TargetingFilterTask_ActorClass.cpp
@@ -0,0 +1,39 @@
+// Copyright Epic Games, Inc. All Rights Reserved.
+#include "TargetingFilterTask_ActorClass.h"
+
+#include "GameFramework/Actor.h"
+
+
+UTargetingFilterTask_ActorClass::UTargetingFilterTask_ActorClass(const FObjectInitializer& ObjectInitializer)
+: Super(ObjectInitializer)
+{
+}
+
+bool UTargetingFilterTask_ActorClass::ShouldFilterTarget(const FTargetingRequestHandle& TargetingHandle, const FTargetingDefaultResultData& TargetData) const
+{
+	if (AActor* TargetActor = TargetData.HitResult.GetActor())
+	{
+		// if the target is one of these classes, filter it out
+		for (UClass* ClassFilter : IgnoredActorClassFilters)
+		{
+			if (TargetActor->IsA(ClassFilter))
+			{
+				return true;
+			}
+		}
+
+		// if the target is one of these classes, do not filter it out
+		for (UClass* ClassFilter : RequiredActorClassFilters)
+		{
+			if (TargetActor->IsA(ClassFilter))
+			{
+				return false;
+			}
+		}
+
+		// if we do not have required class filters, we do NOT want to filter this target
+		return (RequiredActorClassFilters.Num() > 0);
+	}
+
+	return true;
+}
diff --git a/Engine/Plugins/Experimental/GameplayTargetingSystem/Source/GameplayTargetingSystem/Tasks/TargetingFilterTask_ActorClass.h b/Engine/Plugins/Experimental/GameplayTargetingSystem/Source/GameplayTargetingSystem/Tasks/TargetingFilterTask_ActorClass.h
new file mode 100644
--- /dev/null
+++ b/Engine/Plugins/Experimental/GameplayTargetingSystem/Source/GameplayTargetingSystem/Tasks/TargetingFilterTask_ActorClass.h
@@ -0,0 +1,36 @@
+// Copyright Epic Games, Inc. All Rights Reserved.
+#pragma once
+
+#include "GameplayTargetingSystem/Tasks/TargetingFilterTask_BasicFilterTemplate.h"
+#include "GameplayTargetingSystem/Types/TargetingSystemTypes.h"
+
+#include "TargetingFilterTask_ActorClass.generated.h"
+
+
+/**
+*	@class UTargetingFilterTask_ActorClass
+*
+*	Simple filtering task where we check the targets class type against the
+*	required and ignored class types.
+*/
+UCLASS(Blueprintable)
+class UTargetingFilterTask_ActorClass : public UTargetingFilterTask_BasicFilterTemplate
+{
+	GENERATED_BODY()
+
+public:
+	UTargetingFilterTask_ActorClass(const FObjectInitializer& ObjectInitializer);
+
+	/** Evaluation function called by derived classes to process the targeting request */
+	virtual bool ShouldFilterTarget(const FTargetingRequestHandle& TargetingHandle, const FTargetingDefaultResultData& TargetData) const override;
+
+protected:
+	/** The set of required actor classes (must be one of these types to not be filtered out) */
+	UPROPERTY(EditAnywhere, Category = "Targeting Filter | Class", Meta = (AllowAbstract=true))
+	TArray<TObjectPtr<UClass>> RequiredActorClassFilters;
+
+	/** The set of ignored actor classes (must NOT be one of these types) */
+	UPROPERTY(EditAnywhere, Category = "Targeting Filter | Class", Meta = (AllowAbstract = true))
+	TArray<TObjectPtr<UClass>> IgnoredActorClassFilters;
+};
+
diff --git a/Engine/Plugins/Experimental/GameplayTargetingSystem/Source/GameplayTargetingSystem/Tasks/TargetingFilterTask_BasicFilterTemplate.cpp b/Engine/Plugins/Experimental/GameplayTargetingSystem/Source/GameplayTargetingSystem/Tasks/TargetingFilterTask_BasicFilterTemplate.cpp
new file mode 100644
--- /dev/null
+++ b/Engine/Plugins/Experimental/GameplayTargetingSystem/Source/GameplayTargetingSystem/Tasks/TargetingFilterTask_BasicFilterTemplate.cpp
@@ -0,0 +1,113 @@
+// Copyright Epic Games, Inc. All Rights Reserved.
+#include "TargetingFilterTask_BasicFilterTemplate.h"
+
+#include "GameFramework/Actor.h"
+#include "GameplayTargetingSystem/TargetingSystem/TargetingSubsystem.h"
+#include "GameplayTargetingSystem/Types/TargetingSystemTypes.h"
+
+#if ENABLE_DRAW_DEBUG
+#include "Engine/Canvas.h"
+#endif // ENABLE_DRAW_DEBUG
+
+
+UTargetingFilterTask_BasicFilterTemplate::UTargetingFilterTask_BasicFilterTemplate(const FObjectInitializer& ObjectInitializer)
+: Super(ObjectInitializer)
+{
+}
+
+void UTargetingFilterTask_BasicFilterTemplate::Execute(const FTargetingRequestHandle& TargetingHandle) const
+{
+	Super::Execute(TargetingHandle);
+
+	SetTaskAsyncState(TargetingHandle, ETargetingTaskAsyncState::Executing);
+
+#if ENABLE_DRAW_DEBUG
+	ResetFilteredTarget(TargetingHandle);
+#endif // ENABLE_DRAW_DEBUG
+
+	if (TargetingHandle.IsValid())
+	{
+		if (FTargetingDefaultResultsSet* ResultData = FTargetingDefaultResultsSet::Find(TargetingHandle))
+		{
+			const int32 NumTargets = ResultData->TargetResults.Num();
+			for (int32 TargetIterator = NumTargets - 1; TargetIterator >= 0; --TargetIterator)
+			{
+				const FTargetingDefaultResultData& TargetResult = ResultData->TargetResults[TargetIterator];
+				if (ShouldFilterTarget(TargetingHandle, TargetResult))
+				{
+#if ENABLE_DRAW_DEBUG
+					AddFilteredTarget(TargetingHandle, TargetResult);
+#endif // ENABLE_DRAW_DEBUG
+
+					const bool bAllowShrinking = false;
+					ResultData->TargetResults.RemoveAtSwap(TargetIterator, 1, bAllowShrinking);
+				}
+			}
+		}
+	}
+
+	SetTaskAsyncState(TargetingHandle, ETargetingTaskAsyncState::Completed);
+}
+
+bool UTargetingFilterTask_BasicFilterTemplate::ShouldFilterTarget(const FTargetingRequestHandle& TargetingHandle, const FTargetingDefaultResultData& TargetData) const
+{
+	return false;
+}
+
+#if ENABLE_DRAW_DEBUG
+
+void UTargetingFilterTask_BasicFilterTemplate::DrawDebug(UTargetingSubsystem* TargetingSubsystem, FTargetingDebugInfo& Info, const FTargetingRequestHandle& TargetingHandle, float XOffset, float YOffset, int32 MinTextRowsToAdvance) const
+{
+#if WITH_EDITORONLY_DATA
+	if (UTargetingSubsystem::IsTargetingDebugEnabled())
+	{
+		FTargetingDebugData& DebugData = FTargetingDebugData::FindOrAdd(TargetingHandle);
+		FString& ScratchPadString = DebugData.DebugScratchPadStrings.FindOrAdd(GetNameSafe(this));
+		if (!ScratchPadString.IsEmpty())
+		{
+			if (Info.Canvas)
+			{
+				Info.Canvas->SetDrawColor(FColor::Yellow);
+			}
+
+			FString TaskString = FString::Printf(TEXT("Filtered : %s"), *ScratchPadString);
+			TargetingSubsystem->DebugLine(Info, TaskString, XOffset, YOffset, MinTextRowsToAdvance);
+		}
+	}
+#endif // WITH_EDITORONLY_DATA
+}
+
+void UTargetingFilterTask_BasicFilterTemplate::AddFilteredTarget(const FTargetingRequestHandle& TargetingHandle, const FTargetingDefaultResultData& TargetData) const
+{
+#if WITH_EDITORONLY_DATA
+	if (AActor* FilteredActor = TargetData.HitResult.GetActor())
+	{
+		if (FilteredActor && UTargetingSubsystem::IsTargetingDebugEnabled())
+		{
+			FTargetingDebugData& DebugData = FTargetingDebugData::FindOrAdd(TargetingHandle);
+			FString& ScratchPadString = DebugData.DebugScratchPadStrings.FindOrAdd(GetNameSafe(this));
+
+			if (ScratchPadString.IsEmpty())
+			{
+				ScratchPadString = FString::Printf(TEXT("%s"), *GetNameSafe(FilteredActor));
+			}
+			else
+			{
+				ScratchPadString += FString::Printf(TEXT(", %s"), *GetNameSafe(FilteredActor));
+			}
+		}
+	}
+#endif // WITH_EDITORONLY_DATA
+}
+
+void UTargetingFilterTask_BasicFilterTemplate::ResetFilteredTarget(const FTargetingRequestHandle& TargetingHandle) const
+{
+#if WITH_EDITORONLY_DATA
+	FTargetingDebugData& DebugData = FTargetingDebugData::FindOrAdd(TargetingHandle);
+	FString& ScratchPadString = DebugData.DebugScratchPadStrings.FindOrAdd(GetNameSafe(this));
+	ScratchPadString.Reset();
+#endif // WITH_EDITORONLY_DATA
+}
+
+#endif // ENABLE_DRAW_DEBUG
+
diff --git a/Engine/Plugins/Experimental/GameplayTargetingSystem/Source/GameplayTargetingSystem/Tasks/TargetingFilterTask_BasicFilterTemplate.h b/Engine/Plugins/Experimental/GameplayTargetingSystem/Source/GameplayTargetingSystem/Tasks/TargetingFilterTask_BasicFilterTemplate.h
new file mode 100644
--- /dev/null
+++ b/Engine/Plugins/Experimental/GameplayTargetingSystem/Source/GameplayTargetingSystem/Tasks/TargetingFilterTask_BasicFilterTemplate.h
@@ -0,0 +1,47 @@
+// Copyright Epic Games, Inc. All Rights Reserved.
+#pragma once
+
+#include "TargetingTask.h"
+#include "UObject/Object.h"
+
+#include "TargetingFilterTask_BasicFilterTemplate.generated.h"
+
+class AActor;
+class UTargetingSubsystem;
+struct FTargetingDebugInfo;
+struct FTargetingDefaultResultData;
+struct FTargetingRequestHandle;
+
+
+/**
+*	@class UTargetingFilterTask_BasicFilterTemplate
+*
+*	A base class that has a basic setup struct that a majority of filtering tasks
+*	will find convenient.
+*/
+UCLASS(Abstract)
+class TARGETINGSYSTEM_API UTargetingFilterTask_BasicFilterTemplate : public UTargetingTask
+{
+	GENERATED_BODY()
+
+public:
+	UTargetingFilterTask_BasicFilterTemplate(const FObjectInitializer& ObjectInitializer);
+
+private:
+	/** Evaluation function called by derived classes to process the targeting request */
+	virtual void Execute(const FTargetingRequestHandle& TargetingHandle) const override;
+
+protected:
+	/** Called against every target data to determine if the target should be filtered out */
+	virtual bool ShouldFilterTarget(const FTargetingRequestHandle& TargetingHandle, const FTargetingDefaultResultData& TargetData) const;
+
+	/** Debug Helper Methods */
+#if ENABLE_DRAW_DEBUG
+private:
+	virtual void DrawDebug(UTargetingSubsystem* TargetingSubsystem, FTargetingDebugInfo& Info, const FTargetingRequestHandle& TargetingHandle, float XOffset, float YOffset, int32 MinTextRowsToAdvance) const override;
+	void AddFilteredTarget(const FTargetingRequestHandle& TargetingHandle, const FTargetingDefaultResultData& TargetData) const;
+	void ResetFilteredTarget(const FTargetingRequestHandle& TargetingHandle) const;
+#endif // ENABLE_DRAW_DEBUG
+	/** ~Debug Helper Methods */
+};
+
diff --git a/Engine/Plugins/Experimental/GameplayTargetingSystem/Source/GameplayTargetingSystem/Tasks/TargetingFilterTask_SortByDistance.cpp b/Engine/Plugins/Experimental/GameplayTargetingSystem/Source/GameplayTargetingSystem/Tasks/TargetingFilterTask_SortByDistance.cpp
new file mode 100644
--- /dev/null
+++ b/Engine/Plugins/Experimental/GameplayTargetingSystem/Source/GameplayTargetingSystem/Tasks/TargetingFilterTask_SortByDistance.cpp
@@ -0,0 +1,181 @@
+// Copyright Epic Games, Inc. All Rights Reserved.
+#include "TargetingFilterTask_SortByDistance.h"
+
+#include "GameplayTargetingSystem/TargetingSystem/TargetingSubsystem.h"
+
+#if ENABLE_DRAW_DEBUG
+#include "Engine/Canvas.h"
+#endif // ENABLE_DRAW_DEBUG
+
+namespace SortByDistanceConstants
+{
+	const FString PreSortPrefix = TEXT("PreSort");
+	const FString PostSortPrefix = TEXT("PostSort");
+}
+
+UTargetingFilterTask_SortByDistance::UTargetingFilterTask_SortByDistance(const FObjectInitializer& ObjectInitializer)
+: Super(ObjectInitializer)
+{
+	bAscending = true;
+}
+
+void UTargetingFilterTask_SortByDistance::Execute(const FTargetingRequestHandle& TargetingHandle) const
+{
+	Super::Execute(TargetingHandle);
+
+	SetTaskAsyncState(TargetingHandle, ETargetingTaskAsyncState::Executing);
+
+#if ENABLE_DRAW_DEBUG
+	ResetSortDebugStrings(TargetingHandle);
+#endif // ENABLE_DRAW_DEBUG
+
+	if (TargetingHandle.IsValid())
+	{
+		FVector SourceLocation = FVector::ZeroVector;
+		if (FTargetingSourceContext* SourceContext = FTargetingSourceContext::Find(TargetingHandle))
+		{
+			if (SourceContext->SourceActor)
+			{
+				SourceLocation = SourceContext->SourceActor->GetActorLocation();
+			}
+			else if (!SourceContext->SourceLocation.IsZero())
+			{
+				SourceLocation = SourceContext->SourceLocation;
+			}
+		}
+
+		if (!SourceLocation.IsZero())
+		{
+			if (FTargetingDefaultResultsSet* ResultData = FTargetingDefaultResultsSet::Find(TargetingHandle))
+			{
+#if ENABLE_DRAW_DEBUG
+				BuildPreSortDebugString(TargetingHandle, ResultData->TargetResults);
+#endif // ENABLE_DRAW_DEBUG
+
+				// compute the distances
+				const int32 NumTargets = ResultData->TargetResults.Num();
+				for (int32 TargetIterator = NumTargets - 1; TargetIterator >= 0; --TargetIterator)
+				{
+					FTargetingDefaultResultData& TargetResult = ResultData->TargetResults[TargetIterator];
+					if (AActor* TargetActor = TargetResult.HitResult.GetActor())
+					{
+						FVector TargetLocation = TargetActor->GetActorLocation();
+						TargetResult.Score = FVector::DistSquared(SourceLocation, TargetLocation);
+					}
+				}
+
+				// sort the set
+				ResultData->TargetResults.Sort([this](const FTargetingDefaultResultData& Lhs, const FTargetingDefaultResultData& Rhs)
+					{
+						if (bAscending)
+						{
+							return Lhs.Score < Rhs.Score;
+						}
+						else
+						{
+							return Lhs.Score > Rhs.Score;
+						}
+					});
+
+#if ENABLE_DRAW_DEBUG
+				BuildPostSortDebugString(TargetingHandle, ResultData->TargetResults);
+#endif // ENABLE_DRAW_DEBUG
+			}
+		}
+	}
+
+	SetTaskAsyncState(TargetingHandle, ETargetingTaskAsyncState::Completed);
+}
+
+#if ENABLE_DRAW_DEBUG
+
+void UTargetingFilterTask_SortByDistance::DrawDebug(UTargetingSubsystem* TargetingSubsystem, FTargetingDebugInfo& Info, const FTargetingRequestHandle& TargetingHandle, float XOffset, float YOffset, int32 MinTextRowsToAdvance) const
+{
+#if WITH_EDITORONLY_DATA
+	if (UTargetingSubsystem::IsTargetingDebugEnabled())
+	{
+		FTargetingDebugData& DebugData = FTargetingDebugData::FindOrAdd(TargetingHandle);
+		const FString& PreSortScratchPadString = DebugData.DebugScratchPadStrings.FindOrAdd(SortByDistanceConstants::PreSortPrefix + GetNameSafe(this));
+		const FString& PostSortScratchPadString = DebugData.DebugScratchPadStrings.FindOrAdd(SortByDistanceConstants::PostSortPrefix + GetNameSafe(this));
+		if (!PreSortScratchPadString.IsEmpty() && !PostSortScratchPadString.IsEmpty())
+		{
+			if (Info.Canvas)
+			{
+				Info.Canvas->SetDrawColor(FColor::Yellow);
+			}
+
+			FString TaskString = FString::Printf(TEXT("Initial : %s"), *PreSortScratchPadString);
+			TargetingSubsystem->DebugLine(Info, TaskString, XOffset, YOffset, MinTextRowsToAdvance);
+
+			TaskString = FString::Printf(TEXT("Sorted : %s"), *PostSortScratchPadString);
+			TargetingSubsystem->DebugLine(Info, TaskString, XOffset, YOffset, MinTextRowsToAdvance);
+		}
+	}
+#endif // WITH_EDITORONLY_DATA
+}
+
+void UTargetingFilterTask_SortByDistance::BuildPreSortDebugString(const FTargetingRequestHandle& TargetingHandle, const TArray<FTargetingDefaultResultData>& TargetResults) const
+{
+#if WITH_EDITORONLY_DATA
+	if (UTargetingSubsystem::IsTargetingDebugEnabled())
+	{
+		FTargetingDebugData& DebugData = FTargetingDebugData::FindOrAdd(TargetingHandle);
+		FString& PreSortScratchPadString = DebugData.DebugScratchPadStrings.FindOrAdd(SortByDistanceConstants::PreSortPrefix + GetNameSafe(this));
+
+		for (const FTargetingDefaultResultData& TargetData : TargetResults)
+		{
+			if (const AActor* Target = TargetData.HitResult.GetActor())
+			{
+				if (PreSortScratchPadString.IsEmpty())
+				{
+					PreSortScratchPadString = FString::Printf(TEXT("%s"), *GetNameSafe(Target));
+				}
+				else
+				{
+					PreSortScratchPadString += FString::Printf(TEXT(", %s"), *GetNameSafe(Target));
+				}
+			}
+		}
+	}
+#endif // WITH_EDITORONLY_DATA
+}
+
+void UTargetingFilterTask_SortByDistance::BuildPostSortDebugString(const FTargetingRequestHandle& TargetingHandle, const TArray<FTargetingDefaultResultData>& TargetResults) const
+{
+#if WITH_EDITORONLY_DATA
+	if (UTargetingSubsystem::IsTargetingDebugEnabled())
+	{
+		FTargetingDebugData& DebugData = FTargetingDebugData::FindOrAdd(TargetingHandle);
+		FString& PostSortScratchPadString = DebugData.DebugScratchPadStrings.FindOrAdd(SortByDistanceConstants::PostSortPrefix + GetNameSafe(this));
+
+		for (const FTargetingDefaultResultData& TargetData : TargetResults)
+		{
+			if (const AActor* Target = TargetData.HitResult.GetActor())
+			{
+				if (PostSortScratchPadString.IsEmpty())
+				{
+					PostSortScratchPadString = FString::Printf(TEXT("%s"), *GetNameSafe(Target));
+				}
+				else
+				{
+					PostSortScratchPadString += FString::Printf(TEXT(", %s"), *GetNameSafe(Target));
+				}
+			}
+		}
+	}
+#endif // WITH_EDITORONLY_DATA
+}
+
+void UTargetingFilterTask_SortByDistance::ResetSortDebugStrings(const FTargetingRequestHandle& TargetingHandle) const
+{
+#if WITH_EDITORONLY_DATA
+	FTargetingDebugData& DebugData = FTargetingDebugData::FindOrAdd(TargetingHandle);
+	FString& PreSortScratchPadString = DebugData.DebugScratchPadStrings.FindOrAdd(SortByDistanceConstants::PreSortPrefix + GetNameSafe(this));
+	PreSortScratchPadString.Reset();
+
+	FString& PostSortScratchPadString = DebugData.DebugScratchPadStrings.FindOrAdd(SortByDistanceConstants::PostSortPrefix + GetNameSafe(this));
+	PostSortScratchPadString.Reset();
+#endif // WITH_EDITORONLY_DATA
+}
+
+#endif // ENABLE_DRAW_DEBUG
\ No newline at end of file
diff --git a/Engine/Plugins/Experimental/GameplayTargetingSystem/Source/GameplayTargetingSystem/Tasks/TargetingFilterTask_SortByDistance.h b/Engine/Plugins/Experimental/GameplayTargetingSystem/Source/GameplayTargetingSystem/Tasks/TargetingFilterTask_SortByDistance.h
new file mode 100644
--- /dev/null
+++ b/Engine/Plugins/Experimental/GameplayTargetingSystem/Source/GameplayTargetingSystem/Tasks/TargetingFilterTask_SortByDistance.h
@@ -0,0 +1,49 @@
+// Copyright Epic Games, Inc. All Rights Reserved.
+#pragma once
+
+#include "GameplayTargetingSystem/Types/TargetingSystemTypes.h"
+#include "TargetingTask.h"
+#include "UObject/Object.h"
+
+#include "TargetingFilterTask_SortByDistance.generated.h"
+
+class UTargetingSubsystem;
+struct FTargetingDebugInfo;
+struct FTargetingDefaultResultData;
+struct FTargetingRequestHandle;
+
+
+/**
+*	@class UTargetingFilterTask_SortByDistance
+*
+*	Simple sorting filter based on the distance to the source actor.
+*	Note: This filter will use the FTargetingDefaultResultsSet and 
+*	use the Score factor defined for each target to store the distance
+*	and sort by that value.
+*/
+UCLASS(Blueprintable)
+class UTargetingFilterTask_SortByDistance : public UTargetingTask
+{
+	GENERATED_BODY()
+
+public:
+	UTargetingFilterTask_SortByDistance(const FObjectInitializer& ObjectInitializer);
+
+	/** Evaluation function called by derived classes to process the targeting request */
+	virtual void Execute(const FTargetingRequestHandle& TargetingHandle) const override;
+
+protected:
+	UPROPERTY(EditAnywhere, Category = "Target Sorting | Data")
+	uint8 bAscending : 1;
+
+	/** Debug Helper Methods */
+#if ENABLE_DRAW_DEBUG
+private:
+	virtual void DrawDebug(UTargetingSubsystem* TargetingSubsystem, FTargetingDebugInfo& Info, const FTargetingRequestHandle& TargetingHandle, float XOffset, float YOffset, int32 MinTextRowsToAdvance) const override;
+	void BuildPreSortDebugString(const FTargetingRequestHandle& TargetingHandle, const TArray<FTargetingDefaultResultData>& TargetResults) const;
+	void BuildPostSortDebugString(const FTargetingRequestHandle& TargetingHandle, const TArray<FTargetingDefaultResultData>& TargetResults) const;
+	void ResetSortDebugStrings(const FTargetingRequestHandle& TargetingHandle) const;
+#endif // ENABLE_DRAW_DEBUG
+	/** ~Debug Helper Methods */
+};
+
diff --git a/Engine/Plugins/Experimental/GameplayTargetingSystem/Source/GameplayTargetingSystem/Tasks/TargetingSelectionTask_AOE.cpp b/Engine/Plugins/Experimental/GameplayTargetingSystem/Source/GameplayTargetingSystem/Tasks/TargetingSelectionTask_AOE.cpp
new file mode 100644
--- /dev/null
+++ b/Engine/Plugins/Experimental/GameplayTargetingSystem/Source/GameplayTargetingSystem/Tasks/TargetingSelectionTask_AOE.cpp
@@ -0,0 +1,430 @@
+// Copyright Epic Games, Inc. All Rights Reserved.
+#include "TargetingSelectionTask_AOE.h"
+
+#include "CollisionQueryParams.h"
+#include "Components/PrimitiveComponent.h"
+#include "Engine/CollisionProfile.h"
+#include "Engine/EngineTypes.h"
+#include "Engine/World.h"
+#include "GameFramework/Actor.h"
+#include "GameplayTargetingSystem/TargetingSystem/TargetingSubsystem.h"
+#include "GameplayTargetingSystem/Types/TargetingSystemLogs.h"
+
+#if ENABLE_DRAW_DEBUG
+#include "Engine/Canvas.h"
+#endif // ENABLE_DRAW_DEBUG
+
+
+UTargetingSelectionTask_AOE::UTargetingSelectionTask_AOE(const FObjectInitializer& ObjectInitializer)
+: Super(ObjectInitializer)
+{
+	CollisionChannel = ECC_PhysicsBody;
+	bIgnoreSourceActor = false;
+	bIgnoreInstigatorActor = false;
+}
+
+void UTargetingSelectionTask_AOE::Execute(const FTargetingRequestHandle& TargetingHandle) const
+{
+	Super::Execute(TargetingHandle);
+
+	SetTaskAsyncState(TargetingHandle, ETargetingTaskAsyncState::Executing);
+
+	// @note: There isn't Async Overlap support based on Primitive Component, so even if using async targeting, it will
+	// run this task in "immediate" mode.
+	if (IsAsyncTargetingRequest(TargetingHandle) && (ShapeType != ETargetingAOEShape::SourceComponent))
+	{
+		ExecuteAsyncTrace(TargetingHandle);
+	}
+	else
+	{
+		ExecuteImmediateTrace(TargetingHandle);
+	}
+}
+
+void UTargetingSelectionTask_AOE::ExecuteImmediateTrace(const FTargetingRequestHandle& TargetingHandle) const
+{
+#if ENABLE_DRAW_DEBUG
+	ResetDebugString(TargetingHandle);
+#endif // ENABLE_DRAW_DEBUG
+
+	UWorld* World = GetSourceContextWorld(TargetingHandle);
+	if (World && TargetingHandle.IsValid())
+	{
+		const FVector SourceLocation = GetSourceLocation(TargetingHandle) + GetSourceOffset(TargetingHandle);
+
+		TArray<FOverlapResult> OverlapResults;
+		if (ShapeType == ETargetingAOEShape::SourceComponent)
+		{
+			if (const UPrimitiveComponent* CollisionComponent = GetCollisionComponent(TargetingHandle))
+			{
+				FComponentQueryParams ComponentQueryParams(SCENE_QUERY_STAT(UTargetingSelectionTask_AOE_Component));
+				InitCollisionParams(TargetingHandle, ComponentQueryParams);
+				World->ComponentOverlapMulti(OverlapResults, CollisionComponent, CollisionComponent->GetComponentLocation(), CollisionComponent->GetComponentRotation(), ComponentQueryParams);
+			}
+			else
+			{
+				TARGETING_LOG(Warning, TEXT("UTargetingSelectionTask_AOE::Execute - Failed to find a collision component w/ tag [%s] for a SourceComponent ShapeType."), *ComponentTag.ToString());
+			}
+		}
+		else
+		{
+			FCollisionShape CollisionShape = GetCollisionShape();
+			FCollisionQueryParams OverlapParams(TEXT("UTargetingSelectionTask_AOE"), SCENE_QUERY_STAT_ONLY(UTargetingSelectionTask_AOE_Shape), false);
+			InitCollisionParams(TargetingHandle, OverlapParams);
+
+			if (CollisionObjectTypes.Num() > 0)
+			{
+				FCollisionObjectQueryParams ObjectParams;
+				for (auto Iter = CollisionObjectTypes.CreateConstIterator(); Iter; ++Iter)
+				{
+					const ECollisionChannel& Channel = UCollisionProfile::Get()->ConvertToCollisionChannel(false, *Iter);
+					ObjectParams.AddObjectTypesToQuery(Channel);
+				}
+
+				World->OverlapMultiByObjectType(OverlapResults, SourceLocation, FQuat::Identity, ObjectParams, CollisionShape, OverlapParams);
+			}
+			else if (CollisionProfileName.Name != TEXT("NoCollision"))
+			{
+				World->OverlapMultiByProfile(OverlapResults, SourceLocation, FQuat::Identity, CollisionProfileName.Name, CollisionShape, OverlapParams);
+			}
+			else
+			{
+				World->OverlapMultiByChannel(OverlapResults, SourceLocation, FQuat::Identity, CollisionChannel, CollisionShape, OverlapParams);
+			}
+
+#if ENABLE_DRAW_DEBUG
+			if (UTargetingSubsystem::IsTargetingDebugEnabled())
+			{
+				switch (ShapeType)
+				{
+				case ETargetingAOEShape::Box:
+					DrawDebugBox(World, SourceLocation, CollisionShape.GetExtent(), FQuat::Identity, FColor::Red, false, 30.0f, 0, 2.0f);
+					break;
+				case ETargetingAOEShape::Sphere:
+					DrawDebugSphere(World, SourceLocation, CollisionShape.GetSphereRadius(), 32, FColor::Red, false, 30.0f, 0, 2.0f);
+					break;
+				case ETargetingAOEShape::Capsule:
+					DrawDebugCapsule(World, SourceLocation, CollisionShape.GetCapsuleHalfHeight(), CollisionShape.GetCapsuleRadius(), FQuat::Identity, FColor::Red, false, 30.0f, 0, 2.0f);
+					break;
+				case ETargetingAOEShape::Cylinder:
+					DrawDebugBox(World, SourceLocation, CollisionShape.GetExtent(), FQuat::Identity, FColor::Red, false, 30.0f, 0, 1.0f);
+					DrawDebugCylinder(World, -1.0f * CollisionShape.GetExtent(), CollisionShape.GetExtent(), CollisionShape.GetExtent().X, 32, FColor::Yellow, false, 30.0f, 0, 2.0f);
+					break;
+				}
+			}
+#endif // ENABLE_DRAW_DEBUG
+		}
+
+		ProcessOverlapResults(TargetingHandle, OverlapResults);
+	}
+
+	SetTaskAsyncState(TargetingHandle, ETargetingTaskAsyncState::Completed);
+}
+
+void UTargetingSelectionTask_AOE::ExecuteAsyncTrace(const FTargetingRequestHandle& TargetingHandle) const
+{
+	UWorld* World = GetSourceContextWorld(TargetingHandle);
+	if (World && TargetingHandle.IsValid())
+	{
+		const FVector SourceLocation = GetSourceLocation(TargetingHandle) + GetSourceOffset(TargetingHandle);
+
+		FCollisionShape CollisionShape = GetCollisionShape();
+		FCollisionQueryParams OverlapParams(TEXT("UTargetingSelectionTask_AOE"), SCENE_QUERY_STAT_ONLY(UTargetingSelectionTask_AOE_Shape), false);
+		InitCollisionParams(TargetingHandle, OverlapParams);
+
+		FOverlapDelegate Delegate = FOverlapDelegate::CreateUObject(this, &UTargetingSelectionTask_AOE::HandleAsyncOverlapComplete, TargetingHandle);
+		if (CollisionObjectTypes.Num() > 0)
+		{
+			FCollisionObjectQueryParams ObjectParams;
+			for (auto Iter = CollisionObjectTypes.CreateConstIterator(); Iter; ++Iter)
+			{
+				const ECollisionChannel& Channel = UCollisionProfile::Get()->ConvertToCollisionChannel(false, *Iter);
+				ObjectParams.AddObjectTypesToQuery(Channel);
+			}
+
+			World->AsyncOverlapByObjectType(SourceLocation, FQuat::Identity, ObjectParams, CollisionShape, OverlapParams, &Delegate);
+		}
+		else if (CollisionProfileName.Name != TEXT("NoCollision"))
+		{
+			World->AsyncOverlapByProfile(SourceLocation, FQuat::Identity, CollisionProfileName.Name, CollisionShape, OverlapParams, &Delegate);
+		}
+		else
+		{
+			World->AsyncOverlapByChannel(SourceLocation, FQuat::Identity, CollisionChannel, CollisionShape, OverlapParams, FCollisionResponseParams::DefaultResponseParam, &Delegate);
+		}
+	}
+	else
+	{
+		SetTaskAsyncState(TargetingHandle, ETargetingTaskAsyncState::Completed);
+	}
+}
+
+void UTargetingSelectionTask_AOE::HandleAsyncOverlapComplete(const FTraceHandle& InTraceHandle, FOverlapDatum& InOverlapDatum, FTargetingRequestHandle TargetingHandle) const
+{
+	if (TargetingHandle.IsValid())
+	{
+#if ENABLE_DRAW_DEBUG
+		ResetDebugString(TargetingHandle);
+
+		if (UTargetingSubsystem::IsTargetingDebugEnabled())
+		{
+			if (UWorld* World = GetSourceContextWorld(TargetingHandle))
+			{
+				const FVector& SourceLocation = InOverlapDatum.Pos;
+				const FCollisionShape& CollisionShape = InOverlapDatum.CollisionParams.CollisionShape;
+				switch (ShapeType)
+				{
+				case ETargetingAOEShape::Box:
+					DrawDebugBox(World, SourceLocation, CollisionShape.GetExtent(), FQuat::Identity, FColor::Red, false, 30.0f, 0, 2.0f);
+					break;
+				case ETargetingAOEShape::Sphere:
+					DrawDebugSphere(World, SourceLocation, CollisionShape.GetSphereRadius(), 32, FColor::Red, false, 30.0f, 0, 2.0f);
+					break;
+				case ETargetingAOEShape::Capsule:
+					DrawDebugCapsule(World, SourceLocation, CollisionShape.GetCapsuleHalfHeight(), CollisionShape.GetCapsuleRadius(), FQuat::Identity, FColor::Red, false, 30.0f, 0, 2.0f);
+					break;
+				case ETargetingAOEShape::Cylinder:
+					DrawDebugBox(World, SourceLocation, CollisionShape.GetExtent(), FQuat::Identity, FColor::Red, false, 30.0f, 0, 1.0f);
+					DrawDebugCylinder(World, -1.0f * CollisionShape.GetExtent(), CollisionShape.GetExtent(), CollisionShape.GetExtent().X, 32, FColor::Yellow, false, 30.0f, 0, 2.0f);
+					break;
+				}
+			}
+		}
+#endif // ENABLE_DRAW_DEBUG
+
+		ProcessOverlapResults(TargetingHandle, InOverlapDatum.OutOverlaps);
+	}
+
+	SetTaskAsyncState(TargetingHandle, ETargetingTaskAsyncState::Completed);
+}
+
+void UTargetingSelectionTask_AOE::ProcessOverlapResults(const FTargetingRequestHandle& TargetingHandle, const TArray<FOverlapResult>& Overlaps) const
+{
+	// process the overlaps
+	if (Overlaps.Num() > 0)
+	{
+		FTargetingDefaultResultsSet& TargetingResults = FTargetingDefaultResultsSet::FindOrAdd(TargetingHandle);
+		for (const FOverlapResult& OverlapResult : Overlaps)
+		{
+			if (!OverlapResult.GetActor())
+			{
+				continue;
+			}
+
+			// cylinders use box overlaps, so a radius check is necessary to constrain it to the bounds of a cylinder
+			if (ShapeType == ETargetingAOEShape::Cylinder)
+			{
+				const float RadiusSquared = (HalfExtent.X * HalfExtent.X);
+				const FVector SourceLocation = GetSourceLocation(TargetingHandle) + GetSourceOffset(TargetingHandle);
+				const float DistanceSquared = FVector::DistSquared2D(OverlapResult.GetActor()->GetActorLocation(), SourceLocation);
+				if (DistanceSquared > RadiusSquared)
+				{
+					continue;
+				}
+			}
+
+			bool bAddResult = true;
+			for (const FTargetingDefaultResultData& ResultData : TargetingResults.TargetResults)
+			{
+				if (ResultData.HitResult.GetActor() == OverlapResult.GetActor())
+				{
+					bAddResult = false;
+					break;
+				}
+			}
+
+			if (bAddResult)
+			{
+				FTargetingDefaultResultData* ResultData = new(TargetingResults.TargetResults) FTargetingDefaultResultData();
+				ResultData->HitResult.HitObjectHandle = FActorInstanceHandle(OverlapResult.GetActor());
+				ResultData->HitResult.Component = OverlapResult.GetComponent();
+				ResultData->HitResult.bBlockingHit = OverlapResult.bBlockingHit;
+			}
+		}
+
+#if ENABLE_DRAW_DEBUG
+		BuildDebugString(TargetingHandle, TargetingResults.TargetResults);
+#endif // ENABLE_DRAW_DEBUG
+	}
+}
+
+FVector UTargetingSelectionTask_AOE::GetSourceLocation_Implementation(const FTargetingRequestHandle& TargetingHandle) const
+{
+	if (const FTargetingSourceContext* SourceContext = FTargetingSourceContext::Find(TargetingHandle))
+	{
+		if (SourceContext->SourceActor)
+		{
+			return SourceContext->SourceActor->GetActorLocation();
+		}
+		
+		return SourceContext->SourceLocation;
+	}
+
+	return FVector::ZeroVector;
+}
+
+FVector UTargetingSelectionTask_AOE::GetSourceOffset_Implementation(const FTargetingRequestHandle& TargetingHandle) const
+{
+	return DefaultSourceOffset;
+}
+
+FCollisionShape UTargetingSelectionTask_AOE::GetCollisionShape() const
+{
+	switch (ShapeType)
+	{
+	case ETargetingAOEShape::Box:
+		return FCollisionShape::MakeBox(HalfExtent);
+
+	case ETargetingAOEShape::Cylinder:
+		return FCollisionShape::MakeBox(HalfExtent);
+
+	case ETargetingAOEShape::Sphere:
+		return FCollisionShape::MakeSphere(Radius.GetValue());
+
+	case ETargetingAOEShape::Capsule:
+		return FCollisionShape::MakeCapsule(Radius.GetValue(), HalfHeight.GetValue());
+	}
+
+	return FCollisionShape();
+}
+
+const UPrimitiveComponent* UTargetingSelectionTask_AOE::GetCollisionComponent(const FTargetingRequestHandle& TargetingHandle) const
+{
+	if (const FTargetingSourceContext* SourceContext = FTargetingSourceContext::Find(TargetingHandle))
+	{
+		if (SourceContext->SourceActor)
+		{
+			TArray<UPrimitiveComponent*> PrimitiveComponents;
+			SourceContext->SourceActor->GetComponents<UPrimitiveComponent>(PrimitiveComponents);
+
+			for (UPrimitiveComponent* Component : PrimitiveComponents)
+			{
+				if (Component && Component->ComponentHasTag(ComponentTag))
+				{
+					return Component;
+				}
+			}
+		}
+	}
+
+	return nullptr;
+}
+
+void UTargetingSelectionTask_AOE::InitCollisionParams(const FTargetingRequestHandle& TargetingHandle, FCollisionQueryParams& OutParams) const
+{
+	if (const FTargetingSourceContext* SourceContext = FTargetingSourceContext::Find(TargetingHandle))
+	{
+		if (bIgnoreSourceActor && SourceContext->SourceActor)
+		{
+			OutParams.AddIgnoredActor(SourceContext->SourceActor);
+		}
+
+		if (bIgnoreInstigatorActor && SourceContext->InstigatorActor)
+		{
+			OutParams.AddIgnoredActor(SourceContext->InstigatorActor);
+		}
+	}
+}
+
+#if WITH_EDITOR
+bool UTargetingSelectionTask_AOE::CanEditChange(const FProperty* InProperty) const
+{
+	bool bCanEdit = Super::CanEditChange(InProperty);
+
+	if (bCanEdit && InProperty)
+	{
+		if (InProperty->GetFName() == GET_MEMBER_NAME_CHECKED(UTargetingSelectionTask_AOE, CollisionChannel))
+		{
+			return (CollisionProfileName.Name == TEXT("NoCollision") && CollisionObjectTypes.Num() <= 0);
+		}
+
+		if (InProperty->GetFName() == GET_MEMBER_NAME_CHECKED(UTargetingSelectionTask_AOE, CollisionProfileName))
+		{
+			return (CollisionObjectTypes.Num() <= 0);
+		}
+
+		if (InProperty->GetFName() == GET_MEMBER_NAME_CHECKED(UTargetingSelectionTask_AOE, HalfExtent))
+		{
+			return (ShapeType == ETargetingAOEShape::Box || ShapeType == ETargetingAOEShape::Cylinder);
+		}
+
+		if (InProperty->GetFName() == GET_MEMBER_NAME_CHECKED(UTargetingSelectionTask_AOE, Radius))
+		{
+			return (ShapeType == ETargetingAOEShape::Sphere || ShapeType == ETargetingAOEShape::Capsule);
+		}
+
+		if (InProperty->GetFName() == GET_MEMBER_NAME_CHECKED(UTargetingSelectionTask_AOE, HalfHeight))
+		{
+			return (ShapeType == ETargetingAOEShape::Capsule);
+		}
+
+		if (InProperty->GetFName() == GET_MEMBER_NAME_CHECKED(UTargetingSelectionTask_AOE, ComponentTag))
+		{
+			return (ShapeType == ETargetingAOEShape::SourceComponent);
+		}
+	}
+
+	return true;
+}
+#endif // WITH_EDITOR
+
+#if ENABLE_DRAW_DEBUG
+
+void UTargetingSelectionTask_AOE::DrawDebug(UTargetingSubsystem* TargetingSubsystem, FTargetingDebugInfo& Info, const FTargetingRequestHandle& TargetingHandle, float XOffset, float YOffset, int32 MinTextRowsToAdvance) const
+{
+#if WITH_EDITORONLY_DATA
+	if (UTargetingSubsystem::IsTargetingDebugEnabled())
+	{
+		FTargetingDebugData& DebugData = FTargetingDebugData::FindOrAdd(TargetingHandle);
+		FString& ScratchPadString = DebugData.DebugScratchPadStrings.FindOrAdd(GetNameSafe(this));
+		if (!ScratchPadString.IsEmpty())
+		{
+			if (Info.Canvas)
+			{
+				Info.Canvas->SetDrawColor(FColor::Yellow);
+			}
+
+			FString TaskString = FString::Printf(TEXT("Results : %s"), *ScratchPadString);
+			TargetingSubsystem->DebugLine(Info, TaskString, XOffset, YOffset, MinTextRowsToAdvance);
+		}
+	}
+#endif // WITH_EDITORONLY_DATA
+}
+
+void UTargetingSelectionTask_AOE::BuildDebugString(const FTargetingRequestHandle& TargetingHandle, const TArray<FTargetingDefaultResultData>& TargetResults) const
+{
+#if WITH_EDITORONLY_DATA
+	if (UTargetingSubsystem::IsTargetingDebugEnabled())
+	{
+		FTargetingDebugData& DebugData = FTargetingDebugData::FindOrAdd(TargetingHandle);
+		FString& ScratchPadString = DebugData.DebugScratchPadStrings.FindOrAdd(GetNameSafe(this));
+
+		for (const FTargetingDefaultResultData& TargetData : TargetResults)
+		{
+			if (const AActor* Target = TargetData.HitResult.GetActor())
+			{
+				if (ScratchPadString.IsEmpty())
+				{
+					ScratchPadString = FString::Printf(TEXT("%s"), *GetNameSafe(Target));
+				}
+				else
+				{
+					ScratchPadString += FString::Printf(TEXT(", %s"), *GetNameSafe(Target));
+				}
+			}
+		}
+	}
+#endif // WITH_EDITORONLY_DATA
+}
+
+void UTargetingSelectionTask_AOE::ResetDebugString(const FTargetingRequestHandle& TargetingHandle) const
+{
+#if WITH_EDITORONLY_DATA
+	FTargetingDebugData& DebugData = FTargetingDebugData::FindOrAdd(TargetingHandle);
+	FString& ScratchPadString = DebugData.DebugScratchPadStrings.FindOrAdd(GetNameSafe(this));
+	ScratchPadString.Reset();
+#endif // WITH_EDITORONLY_DATA
+}
+
+#endif // ENABLE_DRAW_DEBUG
+
diff --git a/Engine/Plugins/Experimental/GameplayTargetingSystem/Source/GameplayTargetingSystem/Tasks/TargetingSelectionTask_AOE.h b/Engine/Plugins/Experimental/GameplayTargetingSystem/Source/GameplayTargetingSystem/Tasks/TargetingSelectionTask_AOE.h
new file mode 100644
--- /dev/null
+++ b/Engine/Plugins/Experimental/GameplayTargetingSystem/Source/GameplayTargetingSystem/Tasks/TargetingSelectionTask_AOE.h
@@ -0,0 +1,149 @@
+// Copyright Epic Games, Inc. All Rights Reserved.
+#pragma once
+
+#include "CollisionShape.h"
+#include "Engine/CollisionProfile.h"
+#include "Engine/EngineTypes.h"
+#include "GameplayTargetingSystem/Types/TargetingSystemTypes.h"
+#include "ScalableFloat.h"
+#include "TargetingTask.h"
+#include "UObject/Object.h"
+#include "WorldCollision.h"
+
+#include "TargetingSelectionTask_AOE.generated.h"
+
+class UTargetingSubsystem;
+struct FCollisionQueryParams;
+struct FTargetingDebugInfo;
+struct FTargetingRequestHandle;
+
+
+/**
+*	@enum ETargetingAOEShape	
+*/
+UENUM()
+enum class ETargetingAOEShape : uint8
+{
+	Box,
+	Cylinder,
+	Sphere,
+	Capsule,
+	SourceComponent,
+};
+
+/**
+*	@class UTargetingSelectionTask_AOE
+*
+*	Basic AOE based selection task. There are 4 shape types to choose from + a custom option.
+*
+*	Box - Defined by a given half extent
+*	Cylinder - Defined by the half extent (box) + radius
+*	Sphere	- Defined by a given radius
+*	Capsule	- Defined by the radius + half height
+* 
+*	SourceComponent - Use a collision component with a specific component tag as the shape
+*/
+UCLASS(Blueprintable)
+class UTargetingSelectionTask_AOE : public UTargetingTask
+{
+	GENERATED_BODY()
+
+public:
+	UTargetingSelectionTask_AOE(const FObjectInitializer& ObjectInitializer);
+
+	/** Evaluation function called by derived classes to process the targeting request */
+	virtual void Execute(const FTargetingRequestHandle& TargetingHandle) const override;
+
+private:
+	/** Method to process the trace task immediately */
+	void ExecuteImmediateTrace(const FTargetingRequestHandle& TargetingHandle) const;
+
+	/** Method to process the trace task asynchronously */
+	void ExecuteAsyncTrace(const FTargetingRequestHandle& TargetingHandle) const;
+
+	/** Callback for an async overlap */
+	void HandleAsyncOverlapComplete(const FTraceHandle& InTraceHandle, FOverlapDatum& InOverlapDatum, FTargetingRequestHandle TargetingHandle) const;
+
+	/** Method to take the overlap results and store them in the targeting result data */
+	void ProcessOverlapResults(const FTargetingRequestHandle& TargetingHandle, const TArray<FOverlapResult>& Overlaps) const;
+
+protected:
+	/** Native Event to get the source location for the AOE */
+	UFUNCTION(BlueprintNativeEvent, Category = "Target AOE Selection")
+	FVector GetSourceLocation(const FTargetingRequestHandle& TargetingHandle) const;
+
+	/** Native Event to get a source location offset for the AOE */
+	UFUNCTION(BlueprintNativeEvent, Category = "Target AOE Selection")
+	FVector GetSourceOffset(const FTargetingRequestHandle& TargetingHandle) const;
+
+private:
+	/** Helper method to build the Collision Shape */
+	FCollisionShape GetCollisionShape() const;
+
+	/** Helper method to find the custom component defined on the source actor */
+	const UPrimitiveComponent* GetCollisionComponent(const FTargetingRequestHandle& TargetingHandle) const;
+
+	/** Setup CollisionQueryParams for the AOE */
+	void InitCollisionParams(const FTargetingRequestHandle& TargetingHandle, FCollisionQueryParams& OutParams) const;
+
+protected:
+	/** The shape type to use for the AOE */
+	UPROPERTY(EditAnywhere, Category = "Target AOE Selection | Data")
+	ETargetingAOEShape ShapeType = ETargetingAOEShape::Box;
+
+	/** The collision channel to use for the overlap check (as long as Collision Profile Name is not set) */
+	UPROPERTY(EditAnywhere, Category = "Target AOE Selection | Data")
+	TEnumAsByte<ECollisionChannel> CollisionChannel;
+
+	/** The collision profile name to use for the overlap check */
+	UPROPERTY(EditAnywhere, Category = "Target AOE Selection | Data")
+	FCollisionProfileName CollisionProfileName;
+
+	/** The collision profile name to use for the overlap check */
+	UPROPERTY(EditAnywhere, Category = "Target AOE Selection | Data")
+	TArray<TEnumAsByte<EObjectTypeQuery>> CollisionObjectTypes;
+
+	/** The default source location offset used by GetSourceOffset */
+	UPROPERTY(EditAnywhere, Category = "Target AOE Selection | Data")
+	FVector DefaultSourceOffset = FVector::ZeroVector;
+	
+	/** Indicates the trace should ignore the source actor */
+	UPROPERTY(EditAnywhere, Category = "Target AOE Selection | Data")
+	uint8 bIgnoreSourceActor : 1;
+
+	/** Indicates the trace should ignore the source actor */
+	UPROPERTY(EditAnywhere, Category = "Target AOE Selection | Data")
+	uint8 bIgnoreInstigatorActor : 1;
+
+	/** The half extent to use for box and cylinder */
+	UPROPERTY(EditAnywhere, Category = "Target AOE Selection | Box and Cylinder")
+	FVector HalfExtent = FVector::ZeroVector;
+
+	/** The radius to use for sphere and capsule overlaps */
+	UPROPERTY(EditAnywhere, Category = "Target AOE Selection | Sphere and Capsule")
+	FScalableFloat Radius = 0.0f;
+
+	/** The half height to use for capsule overlap checks */
+	UPROPERTY(EditAnywhere, Category = "Target AOE Selection | Capsule")
+	FScalableFloat HalfHeight = 0.0f;
+
+	/** The component tag to use if a custom component is desired as the overlap shape. 
+		Use to look up the component on the source actor */
+	UPROPERTY(EditAnywhere, Category = "Target AOE Selection | Custom Component")
+	FName ComponentTag = NAME_None;
+
+protected:
+#if WITH_EDITOR
+	virtual bool CanEditChange(const FProperty* InProperty) const override;
+#endif
+
+	/** Debug Helper Methods */
+#if ENABLE_DRAW_DEBUG
+private:
+	virtual void DrawDebug(UTargetingSubsystem* TargetingSubsystem, FTargetingDebugInfo& Info, const FTargetingRequestHandle& TargetingHandle, float XOffset, float YOffset, int32 MinTextRowsToAdvance) const override;
+	void BuildDebugString(const FTargetingRequestHandle& TargetingHandle, const TArray<FTargetingDefaultResultData>& TargetResults) const;
+	void ResetDebugString(const FTargetingRequestHandle& TargetingHandle) const;
+#endif // ENABLE_DRAW_DEBUG
+	/** ~Debug Helper Methods */
+};
+
diff --git a/Engine/Plugins/Experimental/GameplayTargetingSystem/Source/GameplayTargetingSystem/Tasks/TargetingSelectionTask_SourceActor.cpp b/Engine/Plugins/Experimental/GameplayTargetingSystem/Source/GameplayTargetingSystem/Tasks/TargetingSelectionTask_SourceActor.cpp
new file mode 100644
--- /dev/null
+++ b/Engine/Plugins/Experimental/GameplayTargetingSystem/Source/GameplayTargetingSystem/Tasks/TargetingSelectionTask_SourceActor.cpp
@@ -0,0 +1,106 @@
+// Copyright Epic Games, Inc. All Rights Reserved.
+#include "TargetingSelectionTask_SourceActor.h"
+
+#include "GameplayTargetingSystem/TargetingSystem/TargetingSubsystem.h"
+
+#if ENABLE_DRAW_DEBUG
+#include "Engine/Canvas.h"
+#endif // ENABLE_DRAW_DEBUG
+
+
+UTargetingSelectionTask_SourceActor::UTargetingSelectionTask_SourceActor(const FObjectInitializer& ObjectInitializer)
+: Super(ObjectInitializer)
+{
+}
+
+void UTargetingSelectionTask_SourceActor::Execute(const FTargetingRequestHandle& TargetingHandle) const
+{
+	Super::Execute(TargetingHandle);
+
+	SetTaskAsyncState(TargetingHandle, ETargetingTaskAsyncState::Executing);
+
+#if ENABLE_DRAW_DEBUG
+	ResetSourceActorDebugString(TargetingHandle);
+#endif // ENABLE_DRAW_DEBUG
+
+	if (TargetingHandle.IsValid())
+	{
+		if (const FTargetingSourceContext* SourceContext = FTargetingSourceContext::Find(TargetingHandle))
+		{
+			if (SourceContext->SourceActor)
+			{
+				bool AddSourceActor = true;
+
+				FTargetingDefaultResultsSet& Results = FTargetingDefaultResultsSet::FindOrAdd(TargetingHandle);
+				for (const FTargetingDefaultResultData& ResultData : Results.TargetResults)
+				{
+					if (ResultData.HitResult.GetActor() == SourceContext->SourceActor)
+					{
+						AddSourceActor = false;
+						break;
+					}
+				}
+
+				if (AddSourceActor)
+				{
+					FTargetingDefaultResultData* ResultData = new(Results.TargetResults) FTargetingDefaultResultData();
+					ResultData->HitResult.HitObjectHandle = FActorInstanceHandle(SourceContext->SourceActor);
+					ResultData->HitResult.Location = SourceContext->SourceActor->GetActorLocation();
+				}
+
+#if ENABLE_DRAW_DEBUG
+				AddSourceActorDebugString(TargetingHandle, SourceContext->SourceActor);
+#endif // ENABLE_DRAW_DEBUG
+			}
+		}
+	}
+
+	SetTaskAsyncState(TargetingHandle, ETargetingTaskAsyncState::Completed);
+}
+
+#if ENABLE_DRAW_DEBUG
+
+void UTargetingSelectionTask_SourceActor::DrawDebug(UTargetingSubsystem* TargetingSubsystem, FTargetingDebugInfo& Info, const FTargetingRequestHandle& TargetingHandle, float XOffset, float YOffset, int32 MinTextRowsToAdvance) const
+{
+#if WITH_EDITORONLY_DATA
+	if (UTargetingSubsystem::IsTargetingDebugEnabled())
+	{
+		FTargetingDebugData& DebugData = FTargetingDebugData::FindOrAdd(TargetingHandle);
+		FString& ScratchPadString = DebugData.DebugScratchPadStrings.FindOrAdd(GetNameSafe(this));
+		if (!ScratchPadString.IsEmpty())
+		{
+			if (Info.Canvas)
+			{
+				Info.Canvas->SetDrawColor(FColor::Yellow);
+			}
+
+			FString TaskString = FString::Printf(TEXT(": %s"), *ScratchPadString);
+			TargetingSubsystem->DebugLine(Info, TaskString, XOffset, YOffset, MinTextRowsToAdvance);
+		}
+	}
+#endif // WITH_EDITORONLY_DATA
+}
+
+void UTargetingSelectionTask_SourceActor::AddSourceActorDebugString(const FTargetingRequestHandle& TargetingHandle, const AActor* SourceActor) const
+{
+#if WITH_EDITORONLY_DATA
+	if (SourceActor && UTargetingSubsystem::IsTargetingDebugEnabled())
+	{
+		FTargetingDebugData& DebugData = FTargetingDebugData::FindOrAdd(TargetingHandle);
+		FString& ScratchPadString = DebugData.DebugScratchPadStrings.FindOrAdd(GetNameSafe(this));
+
+		ScratchPadString = FString::Printf(TEXT("%s"), *GetNameSafe(SourceActor));
+	}
+#endif // WITH_EDITORONLY_DATA
+}
+
+void UTargetingSelectionTask_SourceActor::ResetSourceActorDebugString(const FTargetingRequestHandle& TargetingHandle) const
+{
+#if WITH_EDITORONLY_DATA
+	FTargetingDebugData& DebugData = FTargetingDebugData::FindOrAdd(TargetingHandle);
+	FString& ScratchPadString = DebugData.DebugScratchPadStrings.FindOrAdd(GetNameSafe(this));
+	ScratchPadString.Reset();
+#endif // WITH_EDITORONLY_DATA
+}
+
+#endif // ENABLE_DRAW_DEBUG
diff --git a/Engine/Plugins/Experimental/GameplayTargetingSystem/Source/GameplayTargetingSystem/Tasks/TargetingSelectionTask_SourceActor.h b/Engine/Plugins/Experimental/GameplayTargetingSystem/Source/GameplayTargetingSystem/Tasks/TargetingSelectionTask_SourceActor.h
new file mode 100644
--- /dev/null
+++ b/Engine/Plugins/Experimental/GameplayTargetingSystem/Source/GameplayTargetingSystem/Tasks/TargetingSelectionTask_SourceActor.h
@@ -0,0 +1,42 @@
+// Copyright Epic Games, Inc. All Rights Reserved.
+#pragma once
+
+#include "GameplayTargetingSystem/Types/TargetingSystemTypes.h"
+#include "TargetingTask.h"
+#include "UObject/Object.h"
+
+#include "TargetingSelectionTask_SourceActor.generated.h"
+
+class AActor;
+class UTargetingSubsystem;
+struct FTargetingDebugInfo;
+struct FTargetingRequestHandle;
+
+
+/**
+*	@class UTargetingSelectionTask_SourceActor
+*
+*	Simple targeting selection task that gets the source actor
+*	defined in the targeting source context.
+*/
+UCLASS()
+class UTargetingSelectionTask_SourceActor : public UTargetingTask
+{
+	GENERATED_BODY()
+
+public:
+	UTargetingSelectionTask_SourceActor(const FObjectInitializer& ObjectInitializer);
+
+	/** Evaluation function called by derived classes to process the targeting request */
+	virtual void Execute(const FTargetingRequestHandle& TargetingHandle) const override;
+
+	/** Debug Helper Methods */
+#if ENABLE_DRAW_DEBUG
+private:
+	virtual void DrawDebug(UTargetingSubsystem* TargetingSubsystem, FTargetingDebugInfo& Info, const FTargetingRequestHandle& TargetingHandle, float XOffset, float YOffset, int32 MinTextRowsToAdvance) const override;
+	void AddSourceActorDebugString(const FTargetingRequestHandle& TargetingHandle, const AActor* SourceActor) const;
+	void ResetSourceActorDebugString(const FTargetingRequestHandle& TargetingHandle) const;
+#endif // ENABLE_DRAW_DEBUG
+	/** ~Debug Helper Methods */
+};
+
diff --git a/Engine/Plugins/Experimental/GameplayTargetingSystem/Source/GameplayTargetingSystem/Tasks/TargetingSelectionTask_Trace.cpp b/Engine/Plugins/Experimental/GameplayTargetingSystem/Source/GameplayTargetingSystem/Tasks/TargetingSelectionTask_Trace.cpp
new file mode 100644
--- /dev/null
+++ b/Engine/Plugins/Experimental/GameplayTargetingSystem/Source/GameplayTargetingSystem/Tasks/TargetingSelectionTask_Trace.cpp
@@ -0,0 +1,362 @@
+// Copyright Epic Games, Inc. All Rights Reserved.
+#include "TargetingSelectionTask_Trace.h"
+
+#include "CollisionQueryParams.h"
+#include "Components/PrimitiveComponent.h"
+#include "Engine/EngineTypes.h"
+#include "Engine/World.h"
+#include "GameFramework/Actor.h"
+#include "GameFramework/Pawn.h"
+#include "GameplayTargetingSystem/TargetingSystem/TargetingSubsystem.h"
+#include "GameplayTargetingSystem/Types/TargetingSystemLogs.h"
+
+#if ENABLE_DRAW_DEBUG
+#include "Engine/Canvas.h"
+#endif // ENABLE_DRAW_DEBUG
+
+
+UTargetingSelectionTask_Trace::UTargetingSelectionTask_Trace(const FObjectInitializer& ObjectInitializer)
+: Super(ObjectInitializer)
+{
+	bComplexTrace = false;
+	bIgnoreSourceActor = false;
+	bIgnoreInstigatorActor = false;
+}
+
+void UTargetingSelectionTask_Trace::Execute(const FTargetingRequestHandle& TargetingHandle) const
+{
+	Super::Execute(TargetingHandle);
+
+	SetTaskAsyncState(TargetingHandle, ETargetingTaskAsyncState::Executing);
+
+	if (IsAsyncTargetingRequest(TargetingHandle))
+	{
+		ExecuteAsyncTrace(TargetingHandle);
+	}
+	else
+	{
+		ExecuteImmediateTrace(TargetingHandle);
+	}
+}
+
+FVector UTargetingSelectionTask_Trace::GetSourceLocation_Implementation(const FTargetingRequestHandle& TargetingHandle) const
+{
+	if (const FTargetingSourceContext* SourceContext = FTargetingSourceContext::Find(TargetingHandle))
+	{
+		if (SourceContext->SourceActor)
+		{
+			return SourceContext->SourceActor->GetActorLocation();
+		}
+		
+		return SourceContext->SourceLocation;
+	}
+
+	return FVector::ZeroVector;
+}
+
+FVector UTargetingSelectionTask_Trace::GetSourceOffset_Implementation(const FTargetingRequestHandle& TargetingHandle) const
+{
+	return DefaultSourceOffset;
+}
+
+FVector UTargetingSelectionTask_Trace::GetTraceDirection_Implementation(const FTargetingRequestHandle& TargetingHandle) const
+{
+	if (ExplicitTraceDirection != FVector::ZeroVector)
+	{
+		return ExplicitTraceDirection;
+	}
+	
+	if (const FTargetingSourceContext* SourceContext = FTargetingSourceContext::Find(TargetingHandle))
+	{
+		if (SourceContext->SourceActor)
+		{
+			if (APawn* Pawn = Cast<APawn>(SourceContext->SourceActor))
+			{
+				return Pawn->GetControlRotation().Vector();
+			}
+			else
+			{
+				return SourceContext->SourceActor->GetActorForwardVector();
+			}
+		}
+	}
+
+	return FVector::ZeroVector;
+}
+
+float UTargetingSelectionTask_Trace::GetTraceLength_Implementation(const FTargetingRequestHandle& TargetingHandle) const
+{
+	return DefaultTraceLength.GetValue();
+}
+
+float UTargetingSelectionTask_Trace::GetSweptTraceRadius_Implementation(const FTargetingRequestHandle& TargetingHandle) const
+{
+	return DefaultSweptTraceRadius.GetValue();
+}
+
+void UTargetingSelectionTask_Trace::GetAdditionalActorsToIgnore_Implementation(const FTargetingRequestHandle& TargetingHandle, TArray<AActor*>& OutAdditionalActorsToIgnore) const
+{
+}
+
+void UTargetingSelectionTask_Trace::ExecuteImmediateTrace(const FTargetingRequestHandle& TargetingHandle) const
+{
+	if (UWorld* World = GetSourceContextWorld(TargetingHandle))
+	{
+#if ENABLE_DRAW_DEBUG
+		ResetTraceResultsDebugString(TargetingHandle);
+#endif // ENABLE_DRAW_DEBUG
+
+		const FVector Start = (GetSourceLocation(TargetingHandle) + GetSourceOffset(TargetingHandle));
+		const FVector End = Start + (GetTraceDirection(TargetingHandle) * GetTraceLength(TargetingHandle));
+
+		FCollisionQueryParams Params(SCENE_QUERY_STAT(ExecuteImmediateTrace), bComplexTrace);
+		InitCollisionParams(TargetingHandle, Params);
+
+		TArray<FHitResult> Hits;
+		if (CollisionProfileName.Name != TEXT("NoCollision"))
+		{
+			if (TraceType == ETargetingTraceType::Sweep)
+			{
+				World->SweepMultiByProfile(Hits, Start, End, FQuat::Identity, CollisionProfileName.Name, FCollisionShape::MakeSphere(GetSweptTraceRadius(TargetingHandle)), Params);
+			}
+			else
+			{
+				World->LineTraceMultiByProfile(Hits, Start, End, CollisionProfileName.Name, Params);
+			}
+		}
+		else
+		{
+			ECollisionChannel CollisionChannel = UEngineTypes::ConvertToCollisionChannel(TraceChannel);
+			if (TraceType == ETargetingTraceType::Sweep)
+			{
+				World->SweepMultiByChannel(Hits, Start, End, FQuat::Identity, CollisionChannel, FCollisionShape::MakeSphere(GetSweptTraceRadius(TargetingHandle)), Params);
+			}
+			else
+			{
+				World->LineTraceMultiByChannel(Hits, Start, End, CollisionChannel, Params);
+			}
+		}
+
+#if ENABLE_DRAW_DEBUG
+		if (UTargetingSubsystem::IsTargetingDebugEnabled())
+		{
+			DrawDebugLine(World, Start, End, FColor::Green, false, 30.0f, 0, 2.0f);
+		}
+#endif // ENABLE_DRAW_DEBUG
+
+		ProcessHitResults(TargetingHandle, Hits);
+	}
+
+	SetTaskAsyncState(TargetingHandle, ETargetingTaskAsyncState::Completed);
+}
+
+void UTargetingSelectionTask_Trace::ExecuteAsyncTrace(const FTargetingRequestHandle& TargetingHandle) const
+{
+	if (UWorld* World = GetSourceContextWorld(TargetingHandle))
+	{
+		AActor* SourceActor = nullptr;
+		if (const FTargetingSourceContext* SourceContext = FTargetingSourceContext::Find(TargetingHandle))
+		{
+			SourceActor = SourceContext->SourceActor;
+		}
+		const FVector Start = (GetSourceLocation(TargetingHandle) + GetSourceOffset(TargetingHandle));
+		const FVector End = Start + (GetTraceDirection(TargetingHandle) * GetTraceLength(TargetingHandle));
+
+		FCollisionQueryParams Params(SCENE_QUERY_STAT(ExecuteAsyncTrace), bComplexTrace);
+		InitCollisionParams(TargetingHandle, Params);
+
+		FTraceDelegate Delegate = FTraceDelegate::CreateUObject(this, &UTargetingSelectionTask_Trace::HandleAsyncTraceComplete, TargetingHandle);
+		if (CollisionProfileName.Name != TEXT("NoCollision"))
+		{
+			if (TraceType == ETargetingTraceType::Sweep)
+			{
+				World->AsyncSweepByProfile(EAsyncTraceType::Multi, Start, End, FQuat::Identity, CollisionProfileName.Name, FCollisionShape::MakeSphere(GetSweptTraceRadius(TargetingHandle)), Params, &Delegate);
+			}
+			else
+			{
+				World->AsyncLineTraceByProfile(EAsyncTraceType::Multi, Start, End, CollisionProfileName.Name, Params, &Delegate);
+			}
+		}
+		else
+		{
+			ECollisionChannel CollisionChannel = UEngineTypes::ConvertToCollisionChannel(TraceChannel);
+			if (TraceType == ETargetingTraceType::Sweep)
+			{
+				World->AsyncSweepByChannel(EAsyncTraceType::Multi, Start, End, FQuat::Identity, CollisionChannel, FCollisionShape::MakeSphere(GetSweptTraceRadius(TargetingHandle)), Params, FCollisionResponseParams::DefaultResponseParam, &Delegate);
+			}
+			else
+			{
+				World->AsyncLineTraceByChannel(EAsyncTraceType::Multi, Start, End, CollisionChannel, Params, FCollisionResponseParams::DefaultResponseParam, &Delegate);
+			}
+		}
+	}
+	else
+	{
+		SetTaskAsyncState(TargetingHandle, ETargetingTaskAsyncState::Completed);
+	}
+}
+
+void UTargetingSelectionTask_Trace::HandleAsyncTraceComplete(const FTraceHandle& InTraceHandle, FTraceDatum& InTraceDatum, FTargetingRequestHandle TargetingHandle) const
+{
+	if (TargetingHandle.IsValid())
+	{
+#if ENABLE_DRAW_DEBUG
+		ResetTraceResultsDebugString(TargetingHandle);
+
+		if (UTargetingSubsystem::IsTargetingDebugEnabled())
+		{
+			if (UWorld* World = GetSourceContextWorld(TargetingHandle))
+			{
+				DrawDebugLine(World, InTraceDatum.Start, InTraceDatum.End, FColor::Green, false, 30.0f, 0, 2.0f);
+			}
+		}
+#endif // ENABLE_DRAW_DEBUG
+
+		ProcessHitResults(TargetingHandle, InTraceDatum.OutHits);
+	}
+
+	SetTaskAsyncState(TargetingHandle, ETargetingTaskAsyncState::Completed);
+}
+
+void UTargetingSelectionTask_Trace::ProcessHitResults(const FTargetingRequestHandle& TargetingHandle, const TArray<FHitResult>& Hits) const
+{
+	if (TargetingHandle.IsValid() && Hits.Num() > 0)
+	{
+		FTargetingDefaultResultsSet& TargetingResults = FTargetingDefaultResultsSet::FindOrAdd(TargetingHandle);
+		for (const FHitResult& HitResult : Hits)
+		{
+			if (!HitResult.GetActor())
+			{
+				continue;
+			}
+
+			bool bAddResult = true;
+			for (const FTargetingDefaultResultData& ResultData : TargetingResults.TargetResults)
+			{
+				if (ResultData.HitResult.GetActor() == HitResult.GetActor())
+				{
+					bAddResult = false;
+					break;
+				}
+			}
+
+			if (bAddResult)
+			{
+				FTargetingDefaultResultData* ResultData = new(TargetingResults.TargetResults) FTargetingDefaultResultData();
+				ResultData->HitResult = HitResult;
+			}
+		}
+
+#if ENABLE_DRAW_DEBUG
+		BuildTraceResultsDebugString(TargetingHandle, TargetingResults.TargetResults);
+#endif // ENABLE_DRAW_DEBUG
+	}
+}
+
+void UTargetingSelectionTask_Trace::InitCollisionParams(const FTargetingRequestHandle& TargetingHandle, FCollisionQueryParams& OutParams) const
+{
+	if (const FTargetingSourceContext* SourceContext = FTargetingSourceContext::Find(TargetingHandle))
+	{
+		if (bIgnoreSourceActor && SourceContext->SourceActor)
+		{
+			OutParams.AddIgnoredActor(SourceContext->SourceActor);
+		}
+
+		if (bIgnoreInstigatorActor && SourceContext->InstigatorActor)
+		{
+			OutParams.AddIgnoredActor(SourceContext->InstigatorActor);
+		}
+
+		TArray<AActor*> AdditionalActorsToIgnoreArray;
+		GetAdditionalActorsToIgnore(TargetingHandle, AdditionalActorsToIgnoreArray);
+
+		if (AdditionalActorsToIgnoreArray.Num() > 0)
+		{
+			OutParams.AddIgnoredActors(AdditionalActorsToIgnoreArray);
+		}
+	}
+}
+
+#if WITH_EDITOR
+bool UTargetingSelectionTask_Trace::CanEditChange(const FProperty* InProperty) const
+{
+	bool bCanEdit = Super::CanEditChange(InProperty);
+
+	if (bCanEdit && InProperty)
+	{
+		if (InProperty->GetFName() == GET_MEMBER_NAME_CHECKED(UTargetingSelectionTask_Trace, DefaultSweptTraceRadius.GetValue()))
+		{
+			return (TraceType == ETargetingTraceType::Sweep);
+		}
+
+		if (InProperty->GetFName() == GET_MEMBER_NAME_CHECKED(UTargetingSelectionTask_Trace, TraceChannel))
+		{
+			return (CollisionProfileName.Name == TEXT("NoCollision"));
+		}
+	}
+
+	return true;
+}
+#endif // WITH_EDITOR
+
+#if ENABLE_DRAW_DEBUG
+
+void UTargetingSelectionTask_Trace::DrawDebug(UTargetingSubsystem* TargetingSubsystem, FTargetingDebugInfo& Info, const FTargetingRequestHandle& TargetingHandle, float XOffset, float YOffset, int32 MinTextRowsToAdvance) const
+{
+#if WITH_EDITORONLY_DATA
+	if (UTargetingSubsystem::IsTargetingDebugEnabled())
+	{
+		FTargetingDebugData& DebugData = FTargetingDebugData::FindOrAdd(TargetingHandle);
+		FString& ScratchPadString = DebugData.DebugScratchPadStrings.FindOrAdd(GetNameSafe(this));
+		if (!ScratchPadString.IsEmpty())
+		{
+			if (Info.Canvas)
+			{
+				Info.Canvas->SetDrawColor(FColor::Yellow);
+			}
+
+			FString TaskString = FString::Printf(TEXT("Results : %s"), *ScratchPadString);
+			TargetingSubsystem->DebugLine(Info, TaskString, XOffset, YOffset, MinTextRowsToAdvance);
+		}
+	}
+#endif // WITH_EDITORONLY_DATA
+}
+
+void UTargetingSelectionTask_Trace::BuildTraceResultsDebugString(const FTargetingRequestHandle& TargetingHandle, const TArray<FTargetingDefaultResultData>& TargetResults) const
+{
+#if WITH_EDITORONLY_DATA
+	if (UTargetingSubsystem::IsTargetingDebugEnabled())
+	{
+		FTargetingDebugData& DebugData = FTargetingDebugData::FindOrAdd(TargetingHandle);
+		FString& ScratchPadString = DebugData.DebugScratchPadStrings.FindOrAdd(GetNameSafe(this));
+
+		for (const FTargetingDefaultResultData& TargetData : TargetResults)
+		{
+			if (const AActor* Target = TargetData.HitResult.GetActor())
+			{
+				if (ScratchPadString.IsEmpty())
+				{
+					ScratchPadString = FString::Printf(TEXT("%s"), *GetNameSafe(Target));
+				}
+				else
+				{
+					ScratchPadString += FString::Printf(TEXT(", %s"), *GetNameSafe(Target));
+				}
+			}
+		}
+	}
+#endif // WITH_EDITORONLY_DATA
+}
+
+void UTargetingSelectionTask_Trace::ResetTraceResultsDebugString(const FTargetingRequestHandle& TargetingHandle) const
+{
+#if WITH_EDITORONLY_DATA
+	FTargetingDebugData& DebugData = FTargetingDebugData::FindOrAdd(TargetingHandle);
+	FString& ScratchPadString = DebugData.DebugScratchPadStrings.FindOrAdd(GetNameSafe(this));
+	ScratchPadString.Reset();
+#endif // WITH_EDITORONLY_DATA
+}
+
+#endif // ENABLE_DRAW_DEBUG
+
+
diff --git a/Engine/Plugins/Experimental/GameplayTargetingSystem/Source/GameplayTargetingSystem/Tasks/TargetingSelectionTask_Trace.h b/Engine/Plugins/Experimental/GameplayTargetingSystem/Source/GameplayTargetingSystem/Tasks/TargetingSelectionTask_Trace.h
new file mode 100644
--- /dev/null
+++ b/Engine/Plugins/Experimental/GameplayTargetingSystem/Source/GameplayTargetingSystem/Tasks/TargetingSelectionTask_Trace.h
@@ -0,0 +1,145 @@
+// Copyright Epic Games, Inc. All Rights Reserved.
+#pragma once
+
+#include "CollisionShape.h"
+#include "Engine/CollisionProfile.h"
+#include "GameplayTargetingSystem/Types/TargetingSystemTypes.h"
+#include "ScalableFloat.h"
+#include "TargetingTask.h"
+#include "UObject/Object.h"
+
+#include "TargetingSelectionTask_Trace.generated.h"
+
+class UTargetingSubsystem;
+struct FCollisionQueryParams;
+struct FTargetingDebugInfo;
+struct FTargetingDefaultResultData;
+struct FTargetingRequestHandle;
+struct FTraceDatum;
+struct FTraceHandle;
+
+
+/**
+*	@enum ETargetingAOEShape	
+*/
+UENUM()
+enum class ETargetingTraceType : uint8
+{
+	Line,
+	Sweep,
+};
+
+/**
+*	@class UTargetingSelectionTask_Trace
+*
+*	Selection task that can perform a synchronous or asynchronous trace (line/sweep)
+*	to find all targets up to the first blocking hit (or its end point).
+*/
+UCLASS(Blueprintable)
+class TARGETINGSYSTEM_API UTargetingSelectionTask_Trace : public UTargetingTask
+{
+	GENERATED_BODY()
+
+public:
+	UTargetingSelectionTask_Trace(const FObjectInitializer& ObjectInitializer);
+
+	/** Evaluation function called by derived classes to process the targeting request */
+	virtual void Execute(const FTargetingRequestHandle& TargetingHandle) const override;
+
+protected:
+	/** Native Event to get the source location for the Trace */
+	UFUNCTION(BlueprintNativeEvent, Category = "Target Trace Selection")
+	FVector GetSourceLocation(const FTargetingRequestHandle& TargetingHandle) const;
+
+	/** Native Event to get a source location offset for the Trace */
+	UFUNCTION(BlueprintNativeEvent, Category = "Target Trace Selection")
+	FVector GetSourceOffset(const FTargetingRequestHandle& TargetingHandle) const;
+
+	/** Native Event to get the direction for the Trace */
+	UFUNCTION(BlueprintNativeEvent, Category = "Target Trace Selection")
+	FVector GetTraceDirection(const FTargetingRequestHandle& TargetingHandle) const;
+
+	/** Native Event to get the length for the Trace */
+	UFUNCTION(BlueprintNativeEvent, Category = "Target Trace Selection")
+	float GetTraceLength(const FTargetingRequestHandle& TargetingHandle) const;
+
+	/** Native Event to get the swept trace radius (only called if bSweptTrace is true) */
+	UFUNCTION(BlueprintNativeEvent, Category = "Target Trace Selection")
+	float GetSweptTraceRadius(const FTargetingRequestHandle& TargetingHandle) const;
+
+	/** Native Event to get additional actors the Trace should ignore */
+	UFUNCTION(BlueprintNativeEvent, Category = "Target Trace Selection")
+	void GetAdditionalActorsToIgnore(const FTargetingRequestHandle& TargetingHandle, TArray<AActor*>& OutAdditionalActorsToIgnore) const;
+
+private:
+	/** Method to process the trace task immediately */
+	void ExecuteImmediateTrace(const FTargetingRequestHandle& TargetingHandle) const;
+
+	/** Method to process the trace task asynchronously */
+	void ExecuteAsyncTrace(const FTargetingRequestHandle& TargetingHandle) const;
+
+	/** Callback for an async trace */
+	void HandleAsyncTraceComplete(const FTraceHandle& InTraceHandle, FTraceDatum& InTraceDatum, FTargetingRequestHandle TargetingHandle) const;
+
+	/** Method to take the hit results and store them in the targeting result data */
+	void ProcessHitResults(const FTargetingRequestHandle& TargetingHandle, const TArray<FHitResult>& Hits) const;
+
+	/** Setup CollisionQueryParams for the trace */
+	void InitCollisionParams(const FTargetingRequestHandle& TargetingHandle, FCollisionQueryParams& OutParams) const;
+
+protected:
+	/** The trace type to use */
+	UPROPERTY(EditAnywhere, Category = "Target Trace Selection | Collision Data")
+	ETargetingTraceType TraceType = ETargetingTraceType::Line;
+
+	/** The trace channel to use */
+	UPROPERTY(EditAnywhere, Category = "Target Trace Selection | Collision Data")
+	TEnumAsByte<ETraceTypeQuery> TraceChannel;
+
+	/** The collision profile name to use instead of trace channel (does not work for async traces) */
+	UPROPERTY(EditAnywhere, Category = "Target Trace Selection | Collision Data")
+	FCollisionProfileName CollisionProfileName;
+
+	/** The default swept trace radius used by GetSweptTraceRadius when the trace type is set to Sweep */
+	UPROPERTY(EditAnywhere, Category = "Target Trace Selection | Swept Data")
+	FScalableFloat DefaultSweptTraceRadius = 10.0f;
+
+	/** The default trace length to use if GetTraceLength is not overridden by a child */
+	UPROPERTY(EditAnywhere, Category = "Target Trace Selection | Trace Data")
+	FScalableFloat DefaultTraceLength = 10.0f;
+
+	/** The default source location offset used by GetSourceOffset */
+	UPROPERTY(EditAnywhere, Category = "Target Trace Selection | Trace Data")
+	FVector DefaultSourceOffset = FVector::ZeroVector;
+
+	/** An explicit trace direction to use (default uses pawn control rotation or actor forward vector in GetTraceDirection) */
+	UPROPERTY(EditAnywhere, Category = "Target Trace Selection | Trace Data")
+	FVector ExplicitTraceDirection = FVector::ZeroVector;
+
+	/** Indicates the trace should perform a complex trace */
+	UPROPERTY(EditAnywhere, Category = "Target Trace Selection | Trace Data")
+	uint8 bComplexTrace : 1;
+
+	/** Indicates the trace should ignore the source actor */
+	UPROPERTY(EditAnywhere, Category = "Target Trace Selection | Trace Data")
+	uint8 bIgnoreSourceActor : 1;
+
+	/** Indicates the trace should ignore the source actor */
+	UPROPERTY(EditAnywhere, Category = "Target Trace Selection | Trace Data")
+	uint8 bIgnoreInstigatorActor : 1;
+
+protected:
+#if WITH_EDITOR
+	virtual bool CanEditChange(const FProperty* InProperty) const override;
+#endif
+
+	/** Debug Helper Methods */
+#if ENABLE_DRAW_DEBUG
+private:
+	virtual void DrawDebug(UTargetingSubsystem* TargetingSubsystem, FTargetingDebugInfo& Info, const FTargetingRequestHandle& TargetingHandle, float XOffset, float YOffset, int32 MinTextRowsToAdvance) const override;
+	void BuildTraceResultsDebugString(const FTargetingRequestHandle& TargetingHandle, const TArray<FTargetingDefaultResultData>& TargetResults) const;
+	void ResetTraceResultsDebugString(const FTargetingRequestHandle& TargetingHandle) const;
+#endif // ENABLE_DRAW_DEBUG
+	/** ~Debug Helper Methods */
+};
+
diff --git a/Engine/Plugins/Experimental/GameplayTargetingSystem/Source/GameplayTargetingSystem/Tasks/TargetingTask.cpp b/Engine/Plugins/Experimental/GameplayTargetingSystem/Source/GameplayTargetingSystem/Tasks/TargetingTask.cpp
new file mode 100644
--- /dev/null
+++ b/Engine/Plugins/Experimental/GameplayTargetingSystem/Source/GameplayTargetingSystem/Tasks/TargetingTask.cpp
@@ -0,0 +1,92 @@
+// Copyright Epic Games, Inc. All Rights Reserved.
+#include "TargetingTask.h"
+
+#include "GameFramework/Actor.h"
+#include "GameplayTargetingSystem/Types/TargetingSystemTypes.h"
+
+
+UTargetingTask::UTargetingTask(const FObjectInitializer& ObjectInitializer)
+: Super(ObjectInitializer)
+{
+}
+
+void UTargetingTask::Init(const FTargetingRequestHandle& TargetingHandle) const
+{
+	SetTaskAsyncState(TargetingHandle, ETargetingTaskAsyncState::Initialized);
+}
+
+void UTargetingTask::Execute(const FTargetingRequestHandle& TargetingHandle) const
+{
+}
+
+bool UTargetingTask::IsAsyncTargetingRequest(const FTargetingRequestHandle& TargetingHandle) const
+{
+	if (FTargetingAsyncTaskData* AsyncStateData = FTargetingAsyncTaskData::Find(TargetingHandle))
+	{
+		return AsyncStateData->bAsyncRequest;
+	}
+
+	return false;
+}
+
+void UTargetingTask::SetTaskAsyncState(const FTargetingRequestHandle& TargetingHandle, ETargetingTaskAsyncState AsyncState) const
+{
+	if (FTargetingAsyncTaskData* AsyncStateData = FTargetingAsyncTaskData::Find(TargetingHandle))
+	{
+		if (AsyncStateData->bAsyncRequest)
+		{
+			if (const FTargetingTaskSet** PtrToTaskSet = FTargetingTaskSet::Find(TargetingHandle))
+			{
+				if (const FTargetingTaskSet* TaskSet = (*PtrToTaskSet))
+				{
+					const int32 CurrentTaskIndex = AsyncStateData->CurrentAsyncTaskIndex;
+					if (this == TaskSet->Tasks[CurrentTaskIndex])
+					{
+						AsyncStateData->CurrentAsyncTaskState = AsyncState;
+					}
+				}
+			}
+		}
+	}
+}
+
+ETargetingTaskAsyncState UTargetingTask::GetTaskAsyncState(const FTargetingRequestHandle& TargetingHandle) const
+{
+	if (FTargetingAsyncTaskData* AsyncStateData = FTargetingAsyncTaskData::Find(TargetingHandle))
+	{
+		if (AsyncStateData->bAsyncRequest)
+		{
+			if (const FTargetingTaskSet** PtrToTaskSet = FTargetingTaskSet::Find(TargetingHandle))
+			{
+				if (const FTargetingTaskSet* TaskSet = (*PtrToTaskSet))
+				{
+					const int32 CurrentTaskIndex = AsyncStateData->CurrentAsyncTaskIndex;
+					if (this == TaskSet->Tasks[CurrentTaskIndex])
+					{
+						return AsyncStateData->CurrentAsyncTaskState;
+					}
+				}
+			}
+		}
+	}
+
+	return ETargetingTaskAsyncState::Unitialized;
+}
+
+UWorld* UTargetingTask::GetSourceContextWorld(const FTargetingRequestHandle& TargetingHandle) const
+{
+	if (const FTargetingSourceContext* SourceContext = FTargetingSourceContext::Find(TargetingHandle))
+	{
+		if (SourceContext->SourceActor)
+		{
+			return SourceContext->SourceActor->GetWorld();
+		}
+
+		if (SourceContext->InstigatorActor)
+		{
+			return SourceContext->InstigatorActor->GetWorld();
+		}
+	}
+
+	return GetWorld();
+}
diff --git a/Engine/Plugins/Experimental/GameplayTargetingSystem/Source/GameplayTargetingSystem/Tasks/TargetingTask.h b/Engine/Plugins/Experimental/GameplayTargetingSystem/Source/GameplayTargetingSystem/Tasks/TargetingTask.h
new file mode 100644
--- /dev/null
+++ b/Engine/Plugins/Experimental/GameplayTargetingSystem/Source/GameplayTargetingSystem/Tasks/TargetingTask.h
@@ -0,0 +1,78 @@
+// Copyright Epic Games, Inc. All Rights Reserved.
+#pragma once
+
+#include "DrawDebugHelpers.h"
+#include "GameplayTargetingSystem/Types/TargetingSystemTypes.h"
+#include "UObject/Object.h"
+
+#include "TargetingTask.generated.h"
+
+class UTargetingSubsystem;
+struct FTargetingRequestHandle;
+
+#if ENABLE_DRAW_DEBUG
+struct FTargetingDebugInfo;
+#endif // ENABLE_DRAW_DEBUG
+
+/**
+*	@class UTargetingTask
+*
+*	The base object type that all Targeting Tasks will derive from. The idea
+*	is the targeting system will take target requests that are collections of
+*	target tasks that will potentially generate/remove and perform operations
+*	on sets of targeting results data.
+*
+*	Potential Task Implementations:
+*
+*	Selection Tasks:
+*	Target selection tasks would be used to build up a collection of target
+*	request results. It is recommend they are added first in the targeting
+*	request. Things like ray casts, AOE shapes, actors under a reticle, etc
+*	are cases that generally fall under selection.
+*
+*	Filtering Tasks:
+*	Target filtering tasks are used to reduce the target result data set to
+*	those targets that match a given criteria. Things like actor class, team
+*	distance, facing, etc.
+*
+*	Sorting Tasks:
+*	Target sorting tasks would be useful to take the set and put them in an order
+*	the end user might prefer to make decisions on. Distance (min/max), score rating
+*	etc.
+*/
+UCLASS(EditInlineNew)
+class TARGETINGSYSTEM_API UTargetingTask : public UObject
+{
+	GENERATED_BODY()
+
+public:
+	UTargetingTask(const FObjectInitializer& ObjectInitializer);
+
+	/** Lifecycle function called when the task first begins */
+	virtual void Init(const FTargetingRequestHandle& TargetingHandle) const;
+
+	/** Evaluation function called by derived classes to process the targeting request */
+	virtual void Execute(const FTargetingRequestHandle& TargetingHandle) const;
+
+protected:
+	/** Helper method to check if this task is running in an async targeting request */
+	bool IsAsyncTargetingRequest(const FTargetingRequestHandle& TargetingHandle) const;
+
+	/** Helper method to set the async state for the task (as long as it is the currently running one) */
+	void SetTaskAsyncState(const FTargetingRequestHandle& TargetingHandle, ETargetingTaskAsyncState AsyncState) const;
+
+	/** Helper method to check if a task is currently executing an async operation */
+	ETargetingTaskAsyncState GetTaskAsyncState(const FTargetingRequestHandle& TargetingHandle) const;
+
+	/** Helper method to get the world from the source context (if possible, returns nullptr if one cannot be found) */
+	UWorld* GetSourceContextWorld(const FTargetingRequestHandle& TargetingHandle) const;
+
+	/** Debug Helper Methods */
+#if ENABLE_DRAW_DEBUG
+public:
+	virtual void DrawDebug(UTargetingSubsystem* TargetingSubsystem, FTargetingDebugInfo& Info, const FTargetingRequestHandle& TargetingHandle, float XOffset, float YOffset, int32 MinTextRowsToAdvance = 0) const { }
+#endif // ENABLE_DRAW_DEBUG
+	/** ~Debug Helper Methods */
+};
+
+
diff --git a/Engine/Plugins/Experimental/GameplayTargetingSystem/Source/GameplayTargetingSystem/Types/TargetingSystemDataStores.cpp b/Engine/Plugins/Experimental/GameplayTargetingSystem/Source/GameplayTargetingSystem/Types/TargetingSystemDataStores.cpp
new file mode 100644
--- /dev/null
+++ b/Engine/Plugins/Experimental/GameplayTargetingSystem/Source/GameplayTargetingSystem/Types/TargetingSystemDataStores.cpp
@@ -0,0 +1,23 @@
+// Copyright Epic Games, Inc. All Rights Reserved.
+#include "TargetingSystemDataStores.h"
+
+namespace UE
+{
+	namespace TargetingSystem
+	{
+		/** Data Stores used for Targeting Requests */
+		TTargetingDataStore<FTargetingRequestData> GTargetingRequestDataStore;
+		TTargetingDataStore<const FTargetingTaskSet*> GTargetingTaskSetDataStore;
+		TTargetingDataStore<FTargetingSourceContext> GTargetingSourceDataStore;
+		TTargetingDataStore<FTargetingDefaultResultsSet> GTargetingResultsDataStore;
+		TTargetingDataStore<FTargetingAsyncTaskData> GTargetingAsyncTaskDataStore;
+		TTargetingDataStore<FTargetingImmediateTaskData> GTargetingImmediateTaskDataStore;
+
+#if ENABLE_DRAW_DEBUG
+
+		TTargetingDataStore<FTargetingDebugData> GTargetingDebugDataStore;
+
+#endif // ENABLE_DRAW_DEBUG
+
+	} // TargetingSystem
+} // UE
\ No newline at end of file
diff --git a/Engine/Plugins/Experimental/GameplayTargetingSystem/Source/GameplayTargetingSystem/Types/TargetingSystemDataStores.h b/Engine/Plugins/Experimental/GameplayTargetingSystem/Source/GameplayTargetingSystem/Types/TargetingSystemDataStores.h
new file mode 100644
--- /dev/null
+++ b/Engine/Plugins/Experimental/GameplayTargetingSystem/Source/GameplayTargetingSystem/Types/TargetingSystemDataStores.h
@@ -0,0 +1,218 @@
+// Copyright Epic Games, Inc. All Rights Reserved.
+#pragma once
+
+#include "Containers/SortedMap.h"
+#include "CoreMinimal.h"
+#include "Delegates/DelegateCombinations.h"
+#include "Engine/EngineTypes.h"
+#include "GameplayTargetingSystem/TargetingSystem/TargetingSubsystem.h"
+#include "TargetingSystemTypes.h"
+
+
+namespace UE
+{
+	namespace TargetingSystem
+	{
+		/**
+		*	@struct TTargetingDataStore
+		*
+		*	Templated struct to define targeting data stores. The goal is to provide a flexible
+		*	way for targeting tasks to add/remove/update generic sets of data they wish to work
+		*	with.
+		*
+		*	This allows future growth in games to write tasks either new or extending off base
+		*	framework archetypes to achieve the targeting goals they need.
+		*/
+		template<typename Type>
+		struct TTargetingDataStore
+		{
+			TTargetingDataStore()
+			{
+				ReleaseDelegateHandle = UTargetingSubsystem::ReleaseHandleDelegate.AddStatic(&TTargetingDataStore::OnTargetingRequestHandleReleased);
+			}
+
+			~TTargetingDataStore()
+			{
+				if (ReleaseDelegateHandle.IsValid())
+				{
+					UTargetingSubsystem::ReleaseHandleDelegate.Remove(ReleaseDelegateHandle);
+				}
+			}
+
+			static Type& FindOrAdd(FTargetingRequestHandle Handle);
+			static Type* Find(FTargetingRequestHandle Handle);
+			static void OnTargetingRequestHandleReleased(FTargetingRequestHandle Handle);
+
+		private:
+			// @note: SortMap can be a concern if the Type is a large struct and/or there is high volume of
+			// use at any one time. We could explore other storage/allocation strategies as performance data
+			// dictates.
+
+			/** The Items by handle stored for targeting requests */
+			TSortedMap<FTargetingRequestHandle, Type> Items;
+
+			/** Stores the register handle to the release delegate */
+			FDelegateHandle ReleaseDelegateHandle;
+		};
+
+
+		/** Data Stores used for Targeting Requests */
+
+		/** FTargetingRequestData Data Store */
+		extern TTargetingDataStore<FTargetingRequestData> GTargetingRequestDataStore;
+
+		template<>
+		FORCEINLINE FTargetingRequestData& TTargetingDataStore<FTargetingRequestData>::FindOrAdd(FTargetingRequestHandle Handle)
+		{
+			return GTargetingRequestDataStore.Items.FindOrAdd(Handle);
+		}
+
+		template<>
+		FORCEINLINE FTargetingRequestData* TTargetingDataStore<FTargetingRequestData>::Find(FTargetingRequestHandle Handle)
+		{
+			return GTargetingRequestDataStore.Items.Find(Handle);
+		}
+
+		template<>
+		FORCEINLINE void TTargetingDataStore<FTargetingRequestData>::OnTargetingRequestHandleReleased(FTargetingRequestHandle Handle)
+		{
+			GTargetingRequestDataStore.Items.Remove(Handle);
+		}
+
+
+		/** FTargetingTaskSet Data Store */
+		extern TTargetingDataStore<const FTargetingTaskSet*> GTargetingTaskSetDataStore;
+
+		template<>
+		FORCEINLINE const FTargetingTaskSet*& TTargetingDataStore<const FTargetingTaskSet*>::FindOrAdd(FTargetingRequestHandle Handle)
+		{
+			return GTargetingTaskSetDataStore.Items.FindOrAdd(Handle);
+		}
+
+		template<>
+		FORCEINLINE const FTargetingTaskSet** TTargetingDataStore<const FTargetingTaskSet*>::Find(FTargetingRequestHandle Handle)
+		{
+			return GTargetingTaskSetDataStore.Items.Find(Handle);
+		}
+
+		template<>
+		FORCEINLINE void TTargetingDataStore<const FTargetingTaskSet*>::OnTargetingRequestHandleReleased(FTargetingRequestHandle Handle)
+		{
+			GTargetingTaskSetDataStore.Items.Remove(Handle);
+		}
+
+
+		/** FTargetingSourceContext Data Store */
+		extern TTargetingDataStore<FTargetingSourceContext> GTargetingSourceDataStore;
+
+		template<>
+		FORCEINLINE FTargetingSourceContext& TTargetingDataStore<FTargetingSourceContext>::FindOrAdd(FTargetingRequestHandle Handle)
+		{
+			return GTargetingSourceDataStore.Items.FindOrAdd(Handle);
+		}
+
+		template<>
+		FORCEINLINE FTargetingSourceContext* TTargetingDataStore<FTargetingSourceContext>::Find(FTargetingRequestHandle Handle)
+		{
+			return GTargetingSourceDataStore.Items.Find(Handle);
+		}
+
+		template<>
+		FORCEINLINE void TTargetingDataStore<FTargetingSourceContext>::OnTargetingRequestHandleReleased(FTargetingRequestHandle Handle)
+		{
+			GTargetingSourceDataStore.Items.Remove(Handle);
+		}
+
+
+		/** FTargetingDefaultResultsSet Data Store */
+		extern TTargetingDataStore<FTargetingDefaultResultsSet> GTargetingResultsDataStore;
+
+		template<>
+		FORCEINLINE FTargetingDefaultResultsSet& TTargetingDataStore<FTargetingDefaultResultsSet>::FindOrAdd(FTargetingRequestHandle Handle)
+		{
+			return GTargetingResultsDataStore.Items.FindOrAdd(Handle);
+		}
+
+		template<>
+		FORCEINLINE FTargetingDefaultResultsSet* TTargetingDataStore<FTargetingDefaultResultsSet>::Find(FTargetingRequestHandle Handle)
+		{
+			return GTargetingResultsDataStore.Items.Find(Handle);
+		}
+
+		template<>
+		FORCEINLINE void TTargetingDataStore<FTargetingDefaultResultsSet>::OnTargetingRequestHandleReleased(FTargetingRequestHandle Handle)
+		{
+			GTargetingResultsDataStore.Items.Remove(Handle);
+		}
+
+
+		/** FTargetingAsyncTaskData Data Store */
+		extern TTargetingDataStore<FTargetingAsyncTaskData> GTargetingAsyncTaskDataStore;
+
+		template<>
+		FORCEINLINE FTargetingAsyncTaskData& TTargetingDataStore<FTargetingAsyncTaskData>::FindOrAdd(FTargetingRequestHandle Handle)
+		{
+			return GTargetingAsyncTaskDataStore.Items.FindOrAdd(Handle);
+		}
+
+		template<>
+		FORCEINLINE FTargetingAsyncTaskData* TTargetingDataStore<FTargetingAsyncTaskData>::Find(FTargetingRequestHandle Handle)
+		{
+			return GTargetingAsyncTaskDataStore.Items.Find(Handle);
+		}
+
+		template<>
+		FORCEINLINE void TTargetingDataStore<FTargetingAsyncTaskData>::OnTargetingRequestHandleReleased(FTargetingRequestHandle Handle)
+		{
+			GTargetingAsyncTaskDataStore.Items.Remove(Handle);
+		}
+
+		/** FTargetingImmediateTaskData Data Store */
+		extern TTargetingDataStore<FTargetingImmediateTaskData> GTargetingImmediateTaskDataStore;
+
+		template<>
+		FORCEINLINE FTargetingImmediateTaskData& TTargetingDataStore<FTargetingImmediateTaskData>::FindOrAdd(FTargetingRequestHandle Handle)
+		{
+			return GTargetingImmediateTaskDataStore.Items.FindOrAdd(Handle);
+		}
+
+		template<>
+		FORCEINLINE FTargetingImmediateTaskData* TTargetingDataStore<FTargetingImmediateTaskData>::Find(FTargetingRequestHandle Handle)
+		{
+			return GTargetingImmediateTaskDataStore.Items.Find(Handle);
+		}
+
+		template<>
+		FORCEINLINE void TTargetingDataStore<FTargetingImmediateTaskData>::OnTargetingRequestHandleReleased(FTargetingRequestHandle Handle)
+		{
+			GTargetingImmediateTaskDataStore.Items.Remove(Handle);
+		}
+
+#if ENABLE_DRAW_DEBUG
+
+		/** FTargetingDefaultResults Data Store */
+		extern TTargetingDataStore<FTargetingDebugData> GTargetingDebugDataStore;
+
+		template<>
+		FORCEINLINE FTargetingDebugData& TTargetingDataStore<FTargetingDebugData>::FindOrAdd(FTargetingRequestHandle Handle)
+		{
+			return GTargetingDebugDataStore.Items.FindOrAdd(Handle);
+		}
+
+		template<>
+		FORCEINLINE FTargetingDebugData* TTargetingDataStore<FTargetingDebugData>::Find(FTargetingRequestHandle Handle)
+		{
+			return GTargetingDebugDataStore.Items.Find(Handle);
+		}
+
+		template<>
+		FORCEINLINE void TTargetingDataStore<FTargetingDebugData>::OnTargetingRequestHandleReleased(FTargetingRequestHandle Handle)
+		{
+			GTargetingDebugDataStore.Items.Remove(Handle);
+		}
+
+#endif // ENABLE_DRAW_DEBUG
+
+	} // TargetingSystem
+} // UE
+
diff --git a/Engine/Plugins/Experimental/GameplayTargetingSystem/Source/GameplayTargetingSystem/Types/TargetingSystemLogs.h b/Engine/Plugins/Experimental/GameplayTargetingSystem/Source/GameplayTargetingSystem/Types/TargetingSystemLogs.h
new file mode 100644
--- /dev/null
+++ b/Engine/Plugins/Experimental/GameplayTargetingSystem/Source/GameplayTargetingSystem/Types/TargetingSystemLogs.h
@@ -0,0 +1,11 @@
+// Copyright Epic Games, Inc. All Rights Reserved.
+#pragma once
+
+#include "Logging/LogMacros.h"
+#include "Stats/Stats2.h"
+
+DECLARE_STATS_GROUP(TEXT("TargetingSystem"), STATGROUP_TargetingSystem, STATCAT_Advanced);
+
+DECLARE_LOG_CATEGORY_EXTERN(LogTargetingSystem, Log, All);
+
+#define TARGETING_LOG(Verbosity, Format, ...) UE_LOG(LogTargetingSystem, Verbosity, Format, ##__VA_ARGS__)
diff --git a/Engine/Plugins/Experimental/GameplayTargetingSystem/Source/GameplayTargetingSystem/Types/TargetingSystemTypes.cpp b/Engine/Plugins/Experimental/GameplayTargetingSystem/Source/GameplayTargetingSystem/Types/TargetingSystemTypes.cpp
new file mode 100644
--- /dev/null
+++ b/Engine/Plugins/Experimental/GameplayTargetingSystem/Source/GameplayTargetingSystem/Types/TargetingSystemTypes.cpp
@@ -0,0 +1,132 @@
+// Copyright Epic Games, Inc. All Rights Reserved.
+#include "TargetingSystemTypes.h"
+
+#include "GameplayTargetingSystem/TargetingSystem/TargetingPreset.h"
+#include "GameplayTargetingSystem/TargetingSystem/TargetingSubsystem.h"
+#include "GameplayTargetingSystem/Tasks/TargetingTask.h"
+#include "TargetingSystemLogs.h"
+#include "TargetingSystemDataStores.h"
+
+
+/** @struct FTargetingRequestHandle */
+
+void FTargetingRequestHandle::Reset()
+{
+	Handle = 0;
+}
+
+bool FTargetingRequestHandle::NetSerialize(FArchive& Ar, class UPackageMap* Map, bool& bOutSuccess)
+{
+	TARGETING_LOG(Fatal, TEXT("FTargetingRequestHandle::NetSerialize should not be called!"));
+	return false;
+}
+
+/**	@struct FTargetingTaskSet */
+
+const FTargetingTaskSet*& FTargetingTaskSet::FindOrAdd(FTargetingRequestHandle Handle)
+{
+	return UE::TargetingSystem::GTargetingTaskSetDataStore.FindOrAdd(Handle);
+}
+
+const FTargetingTaskSet** FTargetingTaskSet::Find(FTargetingRequestHandle Handle)
+{
+	return UE::TargetingSystem::GTargetingTaskSetDataStore.Find(Handle);
+}
+
+
+/**	@struct FTargetingDefaultResultsSet */
+
+FTargetingDefaultResultsSet& FTargetingDefaultResultsSet::FindOrAdd(FTargetingRequestHandle Handle)
+{
+	return UE::TargetingSystem::GTargetingResultsDataStore.FindOrAdd(Handle);
+}
+
+FTargetingDefaultResultsSet* FTargetingDefaultResultsSet::Find(FTargetingRequestHandle Handle)
+{
+	return UE::TargetingSystem::GTargetingResultsDataStore.Find(Handle);
+}
+
+
+/**	@struct FTargetingSourceContext */
+
+FTargetingSourceContext& FTargetingSourceContext::FindOrAdd(FTargetingRequestHandle Handle)
+{
+	return UE::TargetingSystem::GTargetingSourceDataStore.FindOrAdd(Handle);
+}
+
+FTargetingSourceContext* FTargetingSourceContext::Find(FTargetingRequestHandle Handle)
+{
+	return UE::TargetingSystem::GTargetingSourceDataStore.Find(Handle);
+}
+
+
+/**	@struct FTargetingRequestData */
+
+FTargetingRequestData& FTargetingRequestData::FindOrAdd(FTargetingRequestHandle Handle)
+{
+	return UE::TargetingSystem::GTargetingRequestDataStore.FindOrAdd(Handle);
+}
+
+
+FTargetingRequestData* FTargetingRequestData::Find(FTargetingRequestHandle Handle)
+{
+	return UE::TargetingSystem::GTargetingRequestDataStore.Find(Handle);
+}
+
+void FTargetingRequestData::Initialize(FTargetingRequestDelegate CompletionDelegate, FTargetingRequestDynamicDelegate CompletionDynamicDelegate)
+{
+	bComplete = false;
+	TargetingRequestDelegate = CompletionDelegate;
+	TargetingRequestDynamicDelegate = CompletionDynamicDelegate;
+}
+
+void FTargetingRequestData::BroadcastTargetingRequestDelegate(FTargetingRequestHandle TargetingRequestHandle)
+{
+	TargetingRequestDelegate.ExecuteIfBound(TargetingRequestHandle);
+	TargetingRequestDynamicDelegate.ExecuteIfBound(TargetingRequestHandle);
+}
+
+
+/** @struct FTargetingAsyncTaskData */
+
+FTargetingAsyncTaskData& FTargetingAsyncTaskData::FindOrAdd(FTargetingRequestHandle Handle)
+{
+	return UE::TargetingSystem::GTargetingAsyncTaskDataStore.FindOrAdd(Handle);
+}
+
+FTargetingAsyncTaskData* FTargetingAsyncTaskData::Find(FTargetingRequestHandle Handle)
+{
+	return UE::TargetingSystem::GTargetingAsyncTaskDataStore.Find(Handle);
+}
+
+void FTargetingAsyncTaskData::InitializeForAsyncProcessing()
+{
+	bAsyncRequest = true;
+	CurrentAsyncTaskIndex = 0;
+	CurrentAsyncTaskState = ETargetingTaskAsyncState::Unitialized;
+}
+
+FTargetingImmediateTaskData& FTargetingImmediateTaskData::FindOrAdd(FTargetingRequestHandle Handle)
+{
+	return UE::TargetingSystem::GTargetingImmediateTaskDataStore.FindOrAdd(Handle);
+}
+
+FTargetingImmediateTaskData* FTargetingImmediateTaskData::Find(FTargetingRequestHandle Handle)
+{
+	return UE::TargetingSystem::GTargetingImmediateTaskDataStore.Find(Handle);
+}
+
+#if WITH_EDITORONLY_DATA
+
+FTargetingDebugData& FTargetingDebugData::FindOrAdd(FTargetingRequestHandle Handle)
+{
+	return UE::TargetingSystem::GTargetingDebugDataStore.FindOrAdd(Handle);
+}
+
+FTargetingDebugData* FTargetingDebugData::Find(FTargetingRequestHandle Handle)
+{
+	return UE::TargetingSystem::GTargetingDebugDataStore.Find(Handle);
+}
+
+#endif // WITH_EDITORONLY_DATA
+
diff --git a/Engine/Plugins/Experimental/GameplayTargetingSystem/Source/GameplayTargetingSystem/Types/TargetingSystemTypes.h b/Engine/Plugins/Experimental/GameplayTargetingSystem/Source/GameplayTargetingSystem/Types/TargetingSystemTypes.h
new file mode 100644
--- /dev/null
+++ b/Engine/Plugins/Experimental/GameplayTargetingSystem/Source/GameplayTargetingSystem/Types/TargetingSystemTypes.h
@@ -0,0 +1,323 @@
+// Copyright Epic Games, Inc. All Rights Reserved.
+#pragma once
+
+#include "CoreMinimal.h"
+#include "DrawDebugHelpers.h"
+#include "Engine/EngineTypes.h"
+#include "Engine/HitResult.h"
+#include "Templates/SharedPointer.h"
+#include "Templates/SharedPointerInternals.h"
+#include "UObject/Class.h"
+#include "UObject/ObjectMacros.h"
+
+#include "TargetingSystemTypes.generated.h"
+
+class AActor;
+class FArchive;
+class UTargetingPreset;
+class UTargetingTask;
+class UPackageMap;
+
+
+/**
+*	@struct FTargetingRequestHandle
+*
+*	The handle that is created when a user wants to make a targeting request.
+*/
+USTRUCT(BlueprintType)
+struct FTargetingRequestHandle
+{
+	GENERATED_BODY()
+
+public:
+	FTargetingRequestHandle()
+	: Handle(0)
+	{ }
+
+	FTargetingRequestHandle(const int32 InHandle)
+	: Handle(InHandle)
+	{ }
+
+	uint32 Handle;
+
+	/** Checks handle validity */
+	FORCEINLINE bool IsValid() const { return (Handle != 0); }
+
+	/** Method to reset the handle */
+	TARGETINGSYSTEM_API void Reset();
+
+	/** We override NetSeralize because we should never be network serializing this handle */
+	bool NetSerialize(FArchive& Ar, class UPackageMap* Map, bool& bOutSuccess);
+
+	/** overloaded operators */
+	FORCEINLINE bool operator==(const FTargetingRequestHandle& InHandle) const { return Handle == InHandle.Handle; }
+	FORCEINLINE bool operator!=(const FTargetingRequestHandle& InHandle) const { return !this->operator==(InHandle); }
+	friend bool operator<(const FTargetingRequestHandle& Lhs, const FTargetingRequestHandle& Rhs)
+	{
+		return Lhs.Handle < Rhs.Handle;
+	}
+};
+
+template<>
+struct TStructOpsTypeTraits<FTargetingRequestHandle> : public TStructOpsTypeTraitsBase2<FTargetingRequestHandle>
+{
+	enum
+	{
+		WithCopy = true,
+		WithNetSerializer = true,
+		WithIdenticalViaEquality = true,
+	};
+};
+
+
+/**
+*	@struct FTargetingTaskSet
+*
+*	A set of tasks to be used by targeting requests to find/processes targets
+*/
+USTRUCT(BlueprintType)
+struct FTargetingTaskSet
+{
+	GENERATED_BODY()
+
+public:
+	/** Convenience method to make using the global data store easier */
+	TARGETINGSYSTEM_API static const FTargetingTaskSet*& FindOrAdd(FTargetingRequestHandle Handle);
+	TARGETINGSYSTEM_API static const FTargetingTaskSet** Find(FTargetingRequestHandle Handle);
+
+	/** The set of tasks that will be used to satisfy a targeting request  */
+	UPROPERTY(EditAnywhere, Instanced, Category="Targeting Task Set")
+	TArray<TObjectPtr<UTargetingTask>> Tasks;
+};
+
+
+/**
+*	@struct FTargetingDefaultResultData
+*/
+USTRUCT()
+struct FTargetingDefaultResultData
+{
+	GENERATED_BODY()
+
+public:
+	FTargetingDefaultResultData()
+	{ }
+
+	/** The hit result for this target */
+	UPROPERTY()
+	FHitResult HitResult;
+
+	/** The score associated w/ this target */
+	float Score = 0.0f;
+};
+
+/**
+*	@struct FTargetingDefaultResultsSet
+*
+*	The base targeting result data used by the tasks implemented at the framework level.
+*	Provides an array of data (hit result / score) that tasks can add/remove/sort/etc
+*	to complete a targeting request.
+*/
+USTRUCT()
+struct FTargetingDefaultResultsSet
+{
+	GENERATED_BODY()
+
+public:
+	FTargetingDefaultResultsSet()
+	{ }
+
+	/** Convenience method to make using the global data store easier */
+	TARGETINGSYSTEM_API static FTargetingDefaultResultsSet& FindOrAdd(FTargetingRequestHandle Handle);
+	TARGETINGSYSTEM_API static FTargetingDefaultResultsSet* Find(FTargetingRequestHandle Handle);
+
+	/** The hit result for this target */
+	UPROPERTY()
+	TArray<FTargetingDefaultResultData> TargetResults;
+};
+
+
+/**
+*	@struct FTargetingSourceContext
+*
+*	Stores context information about a targeting request.
+*/
+USTRUCT(BlueprintType)
+struct FTargetingSourceContext
+{
+	GENERATED_BODY()
+
+public:
+	/** Convenience method to make using the global data store easier */
+	TARGETINGSYSTEM_API static FTargetingSourceContext& FindOrAdd(FTargetingRequestHandle Handle);
+	TARGETINGSYSTEM_API static FTargetingSourceContext* Find(FTargetingRequestHandle Handle);
+
+	/** The optional actor the targeting request sources from (i.e. player/projectile/etc) */
+	UPROPERTY(BlueprintReadWrite, Category = "Targeting Source Context")
+	TObjectPtr<AActor> SourceActor = nullptr;
+
+	/** The optional instigator the targeting request is owned by (i.e. owner of a projectile) */
+	UPROPERTY(BlueprintReadWrite, Category = "Targeting Source Context")
+	TObjectPtr<AActor> InstigatorActor = nullptr;
+
+	/** The optional location the targeting request will source from (i.e. do AOE targeting at x/y/z location) */
+	UPROPERTY(BlueprintReadWrite, Category = "Targeting Source Context")
+	FVector SourceLocation = FVector::ZeroVector;
+
+	/** The optional socket name to use on the source actor (if an actor is defined) */
+	UPROPERTY(BlueprintReadWrite, Category = "Targeting Source Context")
+	FName SourceSocketName = NAME_None;
+
+	/** The optional reference to a source uobject to use in the context */
+	UPROPERTY(BlueprintReadWrite, Category = "Targeting Source Context")
+	TObjectPtr<UObject> SourceObject = nullptr;
+};
+
+
+/** Delegate used by async targeting requests */
+DECLARE_DELEGATE_OneParam(FTargetingRequestDelegate, FTargetingRequestHandle /*TargetingRequestHandle*/);
+DECLARE_DYNAMIC_DELEGATE_OneParam(FTargetingRequestDynamicDelegate, FTargetingRequestHandle, TargetingRequestHandle);
+
+/**
+*	@struct FTargetingRequestData
+*
+*	General purpose targeting request data. Used for general knowledge about
+*	the state of the request.
+*/
+USTRUCT()
+struct FTargetingRequestData
+{
+	GENERATED_BODY()
+
+public:
+	FTargetingRequestData()
+	: bComplete(false)
+	{ }
+
+	/** Convenience method to make using the global data store easier */
+	TARGETINGSYSTEM_API static FTargetingRequestData& FindOrAdd(FTargetingRequestHandle Handle);
+	TARGETINGSYSTEM_API static FTargetingRequestData* Find(FTargetingRequestHandle Handle);
+
+	/** Initializes the targeting request data for async processing */
+	void Initialize(FTargetingRequestDelegate CompletionDelegate, FTargetingRequestDynamicDelegate CompletionDynamicDelegate);
+
+	/** Broadcasts the targeting request delegate */
+	void BroadcastTargetingRequestDelegate(FTargetingRequestHandle TargetingRequestHandle);
+
+	/** Indicates this handle has completed all the targeting request */
+	uint8 bComplete : 1;
+
+	/***/
+	FTargetingRequestDelegate TargetingRequestDelegate;
+
+	/** The dynamic delegate (BP) that fires when the targeting request is complete */
+	UPROPERTY()
+	FTargetingRequestDynamicDelegate TargetingRequestDynamicDelegate;
+};
+
+
+/** @enum ETargetingTaskAsyncState */
+enum class ETargetingTaskAsyncState : uint8
+{
+	Unitialized,	// indicates this task hasn't been started yet
+	Initialized,	// indicates this task has been init'ed and is ready to execute
+	Executing,		// indicates this task is currently processing targets
+	Completed,		// indicates this task is finished processing targets
+};
+
+/**
+*	@struct FTargetingAsyncTaskData
+*
+*	The set of task book keeping data for async targeting requests.
+*/
+USTRUCT()
+struct FTargetingAsyncTaskData
+{
+	GENERATED_BODY()
+
+public:
+	FTargetingAsyncTaskData()
+	: bAsyncRequest(false)
+	, bReleaseOnCompletion(false)
+	, bRequeueOnCompletion(false)
+	{ }
+
+	/** Convenience method to make using the global data store easier */
+	TARGETINGSYSTEM_API static FTargetingAsyncTaskData& FindOrAdd(FTargetingRequestHandle Handle);
+	TARGETINGSYSTEM_API static FTargetingAsyncTaskData* Find(FTargetingRequestHandle Handle);
+
+	/** Initializes the targeting request data for async processing */
+	void InitializeForAsyncProcessing();
+
+	/** The current task index being processed for the targeting request */
+	int32 CurrentAsyncTaskIndex = INDEX_NONE;
+
+	/** The current task's async state */
+	ETargetingTaskAsyncState CurrentAsyncTaskState = ETargetingTaskAsyncState::Unitialized;
+
+	/** Indicates the targeting request was an asynchronous request */
+	uint8 bAsyncRequest : 1;
+
+	/** Indicates the targeting request was an asynchronous request */
+	uint8 bReleaseOnCompletion : 1;
+
+	/** Indicates the target request should be re-queued after completion */
+	uint8 bRequeueOnCompletion : 1;
+};
+
+/**
+*	@struct FTargetingImmediateTaskData
+*
+*	The set of task book keeping data for immediate targeting requests.
+*/
+USTRUCT()
+struct FTargetingImmediateTaskData
+{
+	GENERATED_BODY()
+
+public:
+	FTargetingImmediateTaskData()
+	: bReleaseOnCompletion(false)
+	{ }
+
+	/** Convenience method to make using the global data store easier */
+	TARGETINGSYSTEM_API static FTargetingImmediateTaskData& FindOrAdd(FTargetingRequestHandle Handle);
+	TARGETINGSYSTEM_API static FTargetingImmediateTaskData* Find(FTargetingRequestHandle Handle);
+
+	/** Indicates the targeting request was an asynchronous request */
+	uint8 bReleaseOnCompletion : 1;
+};
+
+
+
+/**
+*	@struct FTargetingDebugData
+*
+*	The set of task book keeping data for async targeting requests.
+*/
+USTRUCT()
+struct FTargetingDebugData
+{
+	GENERATED_BODY()
+
+public:
+	FTargetingDebugData()
+	{ }
+
+#if ENABLE_DRAW_DEBUG
+	/** Convenience method to make using the global data store easier */
+	TARGETINGSYSTEM_API static FTargetingDebugData& FindOrAdd(FTargetingRequestHandle Handle);
+	TARGETINGSYSTEM_API static FTargetingDebugData* Find(FTargetingRequestHandle Handle);
+#endif // ENABLE_DRAW_DEBUG
+
+#if WITH_EDITORONLY_DATA
+	UPROPERTY()
+	TObjectPtr<const UTargetingPreset> TargetingPreset = nullptr;
+
+	UPROPERTY()
+	TMap<FString, FString> DebugScratchPadStrings;
+
+	UPROPERTY()
+	TArray<FTargetingDefaultResultData> CachedTargetResults;
+#endif // WITH_EDITORONLY_DATA
+};
diff --git a/Engine/Plugins/Experimental/GameplayTargetingSystem/TargetingSystem.uplugin b/Engine/Plugins/Experimental/GameplayTargetingSystem/TargetingSystem.uplugin
new file mode 100644
--- /dev/null
+++ b/Engine/Plugins/Experimental/GameplayTargetingSystem/TargetingSystem.uplugin
@@ -0,0 +1,31 @@
+{
+	"FileVersion": 3,
+	"Version": 1,
+	"VersionName": "1.0",
+	"FriendlyName": "TargetingSystem",
+	"Description": "Generic Targeting System to be used for gameplay. Abilities/Aim Assist/etc",
+	"Category": "Gameplay",
+	"CreatedBy": "Epic Games, Inc.",
+	"CreatedByURL": "http://epicgames.com",
+	"DocsURL": "",
+	"MarketplaceURL": "",
+	"SupportURL": "",
+	"EnabledByDefault": false,
+	"CanContainContent": false,
+	"IsBetaVersion": false,
+	"Installed": false,
+	"Modules": [
+		{
+			"Name": "TargetingSystem",
+			"Type": "Runtime",
+			"LoadingPhase": "Default"
+		}
+	],
+	"Plugins": [
+		{
+			"Name": "GameplayAbilities",
+			"Enabled": true
+		}
+	],
+	"IsExperimentalVersion": false
+}
\ No newline at end of file
