Add macros to assist with making custom data stores for Gameplay Targeting System

diff --git a/Engine/Plugins/Experimental/GameplayTargetingSystem/Source/GameplayTargetingSystem/Types/TargetingSystemDataStores.h b/Engine/Plugins/Experimental/GameplayTargetingSystem/Source/GameplayTargetingSystem/Types/TargetingSystemDataStores.h
--- a/Engine/Plugins/Experimental/GameplayTargetingSystem/Source/GameplayTargetingSystem/Types/TargetingSystemDataStores.h
+++ b/Engine/Plugins/Experimental/GameplayTargetingSystem/Source/GameplayTargetingSystem/Types/TargetingSystemDataStores.h
@@ -216,3 +216,57 @@ namespace UE
 	} // TargetingSystem
 } // UE
 
+/** 
+ * Helper macros for creating custom data stores for the targeting system
+ * These data stores can be used to track information linked to a specific
+ * Targeting Request in TargetingTasks.
+ * 
+ * To create a custom data store, place a DECLARE_TARGETING_DATA_STORE
+ * in a header file and a DEFINE_TARGETING_DATA_STORE in a source file,
+ * using the data type to be stored as the macro parameter.
+ * 
+ * For example: A data store that holds a struct FCustomTargetingData
+ * would be declared as follows
+ * 
+ * In CustomTargetingData.h:
+ * DECLARE_TARGETING_DATA_STORE(FCustomTargetingData)
+ * 
+ * In CustomTargetingData.cpp
+ * DEFINE_TARGETING_DATA_STORE(FCustomTargetingData)
+ * 
+ * This new Data Store could be accessed in a relevant TargetingTask like
+ * TTargetingDataStore<FCustomTargetingData>::FindOrAdd(Handle)
+ */ 
+#define DECLARE_TARGETING_DATA_STORE(DataType) namespace UE\
+{\
+	namespace TargetingSystem\
+	{\
+		extern TTargetingDataStore<DataType> GTargetingDataStore##DataType;\
+\
+		template<>\
+		FORCEINLINE DataType& TTargetingDataStore<DataType>::FindOrAdd(FTargetingRequestHandle Handle)\
+		{\
+			return GTargetingDataStore##DataType.Items.FindOrAdd(Handle);\
+		}\
+\
+		template<>\
+		FORCEINLINE DataType* TTargetingDataStore<DataType>::Find(FTargetingRequestHandle Handle)\
+		{\
+			return GTargetingDataStore##DataType.Items.Find(Handle);\
+		}\
+\
+		template<>\
+		FORCEINLINE void TTargetingDataStore<DataType>::OnTargetingRequestHandleReleased(FTargetingRequestHandle Handle)\
+		{\
+			GTargetingDataStore##DataType.Items.Remove(Handle);\
+		}\
+	}\
+}
+
+#define DEFINE_TARGETING_DATA_STORE(DataType) namespace UE\
+{\
+	namespace TargetingSystem\
+	{\
+		TTargetingDataStore<DataType> GTargetingDataStore##DataType;\
+	}\
+}\
