Fixed non-unity/pch errors

diff --git a/Engine/Plugins/Animation/AnimationData/Source/AnimationData/Private/AnimSequencerDataModel.cpp b/Engine/Plugins/Animation/AnimationData/Source/AnimationData/Private/AnimSequencerDataModel.cpp
--- a/Engine/Plugins/Animation/AnimationData/Source/AnimationData/Private/AnimSequencerDataModel.cpp
+++ b/Engine/Plugins/Animation/AnimationData/Source/AnimationData/Private/AnimSequencerDataModel.cpp
@@ -14,6 +14,7 @@
 #include "Rigs/FKControlRig.h"
 #include "Sequencer/MovieSceneControlRigParameterSection.h"
 #include "Sequencer/MovieSceneControlRigParameterTrack.h"
+#include "UObject/LinkerLoad.h"
 #include "UObject/ObjectSaveContext.h"
 #include "Animation/AnimData/AnimDataNotifications.h"
 
diff --git a/Engine/Plugins/Animation/AnimationData/Source/AnimationData/Private/AnimationDataModule.cpp b/Engine/Plugins/Animation/AnimationData/Source/AnimationData/Private/AnimationDataModule.cpp
--- a/Engine/Plugins/Animation/AnimationData/Source/AnimationData/Private/AnimationDataModule.cpp
+++ b/Engine/Plugins/Animation/AnimationData/Source/AnimationData/Private/AnimationDataModule.cpp
@@ -2,6 +2,7 @@
 
 #include "AnimationDataModule.h"
 #include "AnimSequencerDataModel.h"
+#include "Features/IModularFeatures.h"
 #include "ISequencerModule.h"
 #include "Modules/ModuleManager.h"
 
diff --git a/Engine/Plugins/Developer/TraceSourceFiltering/Source/SourceFilteringTrace/Private/TraceWorldFiltering.cpp b/Engine/Plugins/Developer/TraceSourceFiltering/Source/SourceFilteringTrace/Private/TraceWorldFiltering.cpp
--- a/Engine/Plugins/Developer/TraceSourceFiltering/Source/SourceFilteringTrace/Private/TraceWorldFiltering.cpp
+++ b/Engine/Plugins/Developer/TraceSourceFiltering/Source/SourceFilteringTrace/Private/TraceWorldFiltering.cpp
@@ -10,6 +10,7 @@
 
 #include "ObjectTrace.h"
 #include "SourceFilterManager.h"
+#include "UObject/Package.h"
 
 #if WITH_EDITOR
 #include "Editor.h"
diff --git a/Engine/Plugins/Enterprise/VariantManagerContent/Source/VariantManagerContent/Private/ThumbnailGenerator.cpp b/Engine/Plugins/Enterprise/VariantManagerContent/Source/VariantManagerContent/Private/ThumbnailGenerator.cpp
--- a/Engine/Plugins/Enterprise/VariantManagerContent/Source/VariantManagerContent/Private/ThumbnailGenerator.cpp
+++ b/Engine/Plugins/Enterprise/VariantManagerContent/Source/VariantManagerContent/Private/ThumbnailGenerator.cpp
@@ -2,6 +2,7 @@
 
 #include "ThumbnailGenerator.h"
 
+#include "Engine/World.h"
 #include "VariantManagerContentLog.h"
 
 #include "CanvasTypes.h"
diff --git a/Engine/Plugins/Experimental/ChaosCloth/Source/Chaos/Private/ChaosCloth/ChaosClothVisualization.cpp b/Engine/Plugins/Experimental/ChaosCloth/Source/Chaos/Private/ChaosCloth/ChaosClothVisualization.cpp
--- a/Engine/Plugins/Experimental/ChaosCloth/Source/Chaos/Private/ChaosCloth/ChaosClothVisualization.cpp
+++ b/Engine/Plugins/Experimental/ChaosCloth/Source/Chaos/Private/ChaosCloth/ChaosClothVisualization.cpp
@@ -26,6 +26,7 @@
 #include "Chaos/XPBDSpringConstraints.h"
 #include "DynamicMeshBuilder.h"
 #include "Engine/EngineTypes.h"
+#include "SceneManagement.h"
 #if WITH_EDITOR
 #include "Materials/Material.h"
 #include "Engine/Canvas.h"  // For draw text
diff --git a/Engine/Plugins/Experimental/GameFeatures/Source/GameFeatures/Private/GameFeatureData.cpp b/Engine/Plugins/Experimental/GameFeatures/Source/GameFeatures/Private/GameFeatureData.cpp
--- a/Engine/Plugins/Experimental/GameFeatures/Source/GameFeatures/Private/GameFeatureData.cpp
+++ b/Engine/Plugins/Experimental/GameFeatures/Source/GameFeatures/Private/GameFeatureData.cpp
@@ -1,6 +1,7 @@
 // Copyright Epic Games, Inc. All Rights Reserved.
 
 #include "GameFeatureData.h"
+#include "AssetRegistry/AssetData.h"
 #include "Engine/AssetManager.h"
 #include "GameFeaturesSubsystem.h"
 #include "GameFeaturesSubsystemSettings.h"
diff --git a/Engine/Plugins/Experimental/GameplayTargetingSystem/Source/GameplayTargetingSystem/TargetingSystem/TargetingSubsystem.cpp b/Engine/Plugins/Experimental/GameplayTargetingSystem/Source/GameplayTargetingSystem/TargetingSystem/TargetingSubsystem.cpp
--- a/Engine/Plugins/Experimental/GameplayTargetingSystem/Source/GameplayTargetingSystem/TargetingSystem/TargetingSubsystem.cpp
+++ b/Engine/Plugins/Experimental/GameplayTargetingSystem/Source/GameplayTargetingSystem/TargetingSystem/TargetingSubsystem.cpp
@@ -4,6 +4,7 @@
 #include "Engine/World.h"
 #include "GameplayTargetingSystem/Types/TargetingSystemLogs.h"
 #include "GameplayTargetingSystem/Types/TargetingSystemTypes.h"
+#include "HAL/IConsoleManager.h"
 #include "Kismet/KismetStringLibrary.h"
 #include "ProfilingDebugging/CsvProfiler.h"
 #include "Stats/Stats2.h"
diff --git a/Engine/Plugins/Experimental/GameplayTargetingSystem/Source/GameplayTargetingSystem/Tasks/TargetingFilterTask_SortByDistance.cpp b/Engine/Plugins/Experimental/GameplayTargetingSystem/Source/GameplayTargetingSystem/Tasks/TargetingFilterTask_SortByDistance.cpp
--- a/Engine/Plugins/Experimental/GameplayTargetingSystem/Source/GameplayTargetingSystem/Tasks/TargetingFilterTask_SortByDistance.cpp
+++ b/Engine/Plugins/Experimental/GameplayTargetingSystem/Source/GameplayTargetingSystem/Tasks/TargetingFilterTask_SortByDistance.cpp
@@ -1,6 +1,7 @@
 // Copyright Epic Games, Inc. All Rights Reserved.
 #include "TargetingFilterTask_SortByDistance.h"
 
+#include "GameFramework/Actor.h"
 #include "GameplayTargetingSystem/TargetingSystem/TargetingSubsystem.h"
 
 #if ENABLE_DRAW_DEBUG
diff --git a/Engine/Plugins/Experimental/GameplayTargetingSystem/Source/GameplayTargetingSystem/Tasks/TargetingSelectionTask_SourceActor.cpp b/Engine/Plugins/Experimental/GameplayTargetingSystem/Source/GameplayTargetingSystem/Tasks/TargetingSelectionTask_SourceActor.cpp
--- a/Engine/Plugins/Experimental/GameplayTargetingSystem/Source/GameplayTargetingSystem/Tasks/TargetingSelectionTask_SourceActor.cpp
+++ b/Engine/Plugins/Experimental/GameplayTargetingSystem/Source/GameplayTargetingSystem/Tasks/TargetingSelectionTask_SourceActor.cpp
@@ -1,6 +1,7 @@
 // Copyright Epic Games, Inc. All Rights Reserved.
 #include "TargetingSelectionTask_SourceActor.h"
 
+#include "GameFramework/Actor.h"
 #include "GameplayTargetingSystem/TargetingSystem/TargetingSubsystem.h"
 
 #if ENABLE_DRAW_DEBUG
diff --git a/Engine/Plugins/Experimental/PCG/Source/PCG/Private/PCGGraph.cpp b/Engine/Plugins/Experimental/PCG/Source/PCG/Private/PCGGraph.cpp
--- a/Engine/Plugins/Experimental/PCG/Source/PCG/Private/PCGGraph.cpp
+++ b/Engine/Plugins/Experimental/PCG/Source/PCG/Private/PCGGraph.cpp
@@ -5,6 +5,7 @@
 #include "PCGInputOutputSettings.h"
 #include "PCGSubgraph.h"
 #include "PCGSubsystem.h"
+#include "UObject/Package.h"
 
 #if WITH_EDITOR
 #include "Editor.h"
diff --git a/Engine/Plugins/FX/Niagara/Source/Niagara/Private/NiagaraCullProxyComponent.cpp b/Engine/Plugins/FX/Niagara/Source/Niagara/Private/NiagaraCullProxyComponent.cpp
--- a/Engine/Plugins/FX/Niagara/Source/Niagara/Private/NiagaraCullProxyComponent.cpp
+++ b/Engine/Plugins/FX/Niagara/Source/Niagara/Private/NiagaraCullProxyComponent.cpp
@@ -1,6 +1,7 @@
 // Copyright Epic Games, Inc. All Rights Reserved.
 
 #include "NiagaraCullProxyComponent.h"
+#include "NiagaraEffectType.h"
 #include "NiagaraSystem.h"
 
 #include UE_INLINE_GENERATED_CPP_BY_NAME(NiagaraCullProxyComponent)
diff --git a/Engine/Plugins/FX/Niagara/Source/NiagaraEditor/Private/ViewModels/NiagaraSimCacheViewModel.cpp b/Engine/Plugins/FX/Niagara/Source/NiagaraEditor/Private/ViewModels/NiagaraSimCacheViewModel.cpp
--- a/Engine/Plugins/FX/Niagara/Source/NiagaraEditor/Private/ViewModels/NiagaraSimCacheViewModel.cpp
+++ b/Engine/Plugins/FX/Niagara/Source/NiagaraEditor/Private/ViewModels/NiagaraSimCacheViewModel.cpp
@@ -7,6 +7,8 @@
 #include "NiagaraDebuggerCommon.h"
 #include "Serialization/MemoryReader.h"
 #include "Serialization/ObjectAndNameAsStringProxyArchive.h"
+#include "UObject/Package.h"
+#include "UObject/Package.h"
 
 #define LOCTEXT_NAMESPACE "NiagaraSimCacheViewModel"
 
diff --git a/Engine/Plugins/Importers/USDImporter/Source/USDSchemas/Private/MeshTranslationImpl.cpp b/Engine/Plugins/Importers/USDImporter/Source/USDSchemas/Private/MeshTranslationImpl.cpp
--- a/Engine/Plugins/Importers/USDImporter/Source/USDSchemas/Private/MeshTranslationImpl.cpp
+++ b/Engine/Plugins/Importers/USDImporter/Source/USDSchemas/Private/MeshTranslationImpl.cpp
@@ -2,6 +2,7 @@
 
 #include "MeshTranslationImpl.h"
 
+#include "UObject/Package.h"
 #include "USDAssetCache.h"
 #include "USDAssetImportData.h"
 #include "USDGeomMeshConversion.h"
diff --git a/Engine/Plugins/Runtime/GameplayAbilities/Source/GameplayAbilities/Private/Serialization/GameplayEffectContextNetSerializer.cpp b/Engine/Plugins/Runtime/GameplayAbilities/Source/GameplayAbilities/Private/Serialization/GameplayEffectContextNetSerializer.cpp
--- a/Engine/Plugins/Runtime/GameplayAbilities/Source/GameplayAbilities/Private/Serialization/GameplayEffectContextNetSerializer.cpp
+++ b/Engine/Plugins/Runtime/GameplayAbilities/Source/GameplayAbilities/Private/Serialization/GameplayEffectContextNetSerializer.cpp
@@ -5,6 +5,7 @@
 
 #if UE_WITH_IRIS
 
+#include "Engine/HitResult.h"
 #include "GameplayEffectTypes.h"
 #include "Net/Core/NetBitArray.h"
 #include "Net/Core/Trace/NetTrace.h"
diff --git a/Engine/Plugins/VirtualProduction/LevelSnapshots/Source/LevelSnapshots/Private/Data/Util/EquivalenceUtil.cpp b/Engine/Plugins/VirtualProduction/LevelSnapshots/Source/LevelSnapshots/Private/Data/Util/EquivalenceUtil.cpp
--- a/Engine/Plugins/VirtualProduction/LevelSnapshots/Source/LevelSnapshots/Private/Data/Util/EquivalenceUtil.cpp
+++ b/Engine/Plugins/VirtualProduction/LevelSnapshots/Source/LevelSnapshots/Private/Data/Util/EquivalenceUtil.cpp
@@ -18,6 +18,7 @@
 #include "EngineUtils.h"
 #include "LevelSnapshot.h"
 #include "GameFramework/Actor.h"
+#include "UObject/Package.h"
 #include "UObject/TextProperty.h"
 #include "UObject/UnrealType.h"
 #include "WorldData/WorldDataUtil.h"
diff --git a/Engine/Source/Editor/MainFrame/Private/MainFrameModule.cpp b/Engine/Source/Editor/MainFrame/Private/MainFrameModule.cpp
--- a/Engine/Source/Editor/MainFrame/Private/MainFrameModule.cpp
+++ b/Engine/Source/Editor/MainFrame/Private/MainFrameModule.cpp
@@ -1,6 +1,7 @@
 // Copyright Epic Games, Inc. All Rights Reserved.
 
 #include "MainFrameModule.h"
+#include "Features/IModularFeatures.h"
 #include "Widgets/Text/STextBlock.h"
 #include "Widgets/Layout/SBorder.h"
 #include "SlateOptMacros.h"
diff --git a/Engine/Source/Runtime/Engine/Private/Animation/AnimData/IAnimationDataController.cpp b/Engine/Source/Runtime/Engine/Private/Animation/AnimData/IAnimationDataController.cpp
--- a/Engine/Source/Runtime/Engine/Private/Animation/AnimData/IAnimationDataController.cpp
+++ b/Engine/Source/Runtime/Engine/Private/Animation/AnimData/IAnimationDataController.cpp
@@ -1,6 +1,7 @@
 // Copyright Epic Games, Inc. All Rights Reserved.
 
 #include "Animation/AnimData/IAnimationDataController.h"
+#include "UObject/Package.h"
 
 void IAnimationDataController::ReportMessage(const UObject* ErrorObject, const FText& InMessage, ELogVerbosity::Type LogVerbosity)
 {
diff --git a/Engine/Source/Runtime/Engine/Private/AudioDebug.cpp b/Engine/Source/Runtime/Engine/Private/AudioDebug.cpp
--- a/Engine/Source/Runtime/Engine/Private/AudioDebug.cpp
+++ b/Engine/Source/Runtime/Engine/Private/AudioDebug.cpp
@@ -12,6 +12,7 @@
 #include "DrawDebugHelpers.h"
 #include "DSP/Dsp.h"
 #include "Engine/Font.h"
+#include "Engine/GameViewportClient.h"
 #include "Engine/World.h"
 #include "GameFramework/GameUserSettings.h"
 #include "HAL/IConsoleManager.h"
diff --git a/Engine/Source/Runtime/Engine/Private/BlueprintGeneratedClass.cpp b/Engine/Source/Runtime/Engine/Private/BlueprintGeneratedClass.cpp
--- a/Engine/Source/Runtime/Engine/Private/BlueprintGeneratedClass.cpp
+++ b/Engine/Source/Runtime/Engine/Private/BlueprintGeneratedClass.cpp
@@ -1,6 +1,7 @@
 // Copyright Epic Games, Inc. All Rights Reserved.
 
 #include "Engine/BlueprintGeneratedClass.h"
+#include "HAL/IConsoleManager.h"
 #include "Misc/CoreMisc.h"
 #include "Stats/StatsMisc.h"
 #include "UObject/UObjectHash.h"
diff --git a/Engine/Source/Runtime/Engine/Private/GameInstance.cpp b/Engine/Source/Runtime/Engine/Private/GameInstance.cpp
--- a/Engine/Source/Runtime/Engine/Private/GameInstance.cpp
+++ b/Engine/Source/Runtime/Engine/Private/GameInstance.cpp
@@ -34,6 +34,7 @@
 #include "ReplaySubsystem.h"
 #include "GenericPlatform/GenericPlatformInputDeviceMapper.h"
 #include "ProfilingDebugging/CpuProfilerTrace.h"
+#include "UObject/Package.h"
 
 #include UE_INLINE_GENERATED_CPP_BY_NAME(GameInstance)
 
diff --git a/Engine/Source/Runtime/Engine/Private/Net/Iris/ReplicationSystem/ActorReplicationBridge.cpp b/Engine/Source/Runtime/Engine/Private/Net/Iris/ReplicationSystem/ActorReplicationBridge.cpp
--- a/Engine/Source/Runtime/Engine/Private/Net/Iris/ReplicationSystem/ActorReplicationBridge.cpp
+++ b/Engine/Source/Runtime/Engine/Private/Net/Iris/ReplicationSystem/ActorReplicationBridge.cpp
@@ -30,6 +30,7 @@
 #include "Net/Core/Misc/NetSubObjectRegistry.h"
 #include "Net/NetSubObjectRegistryGetter.h"
 #include "Templates/Casts.h"
+#include "UObject/Package.h"
 #include <limits>
 
 #if UE_NET_ALLOW_MULTIPLE_REPLICATION_SYSTEMS
diff --git a/Engine/Source/Runtime/Engine/Private/ObjectTrace.cpp b/Engine/Source/Runtime/Engine/Private/ObjectTrace.cpp
--- a/Engine/Source/Runtime/Engine/Private/ObjectTrace.cpp
+++ b/Engine/Source/Runtime/Engine/Private/ObjectTrace.cpp
@@ -11,6 +11,7 @@
 #include "Trace/Trace.inl"
 #include "UObject/Object.h"
 #include "UObject/Class.h"
+#include "UObject/Package.h"
 #include "UObject/UObjectAnnotation.h"
 #include "UObject/WeakObjectPtrTemplates.h"
 #include "UObject/WeakObjectPtr.h"
diff --git a/Engine/Source/Runtime/Engine/Private/UserDefinedEnum.cpp b/Engine/Source/Runtime/Engine/Private/UserDefinedEnum.cpp
--- a/Engine/Source/Runtime/Engine/Private/UserDefinedEnum.cpp
+++ b/Engine/Source/Runtime/Engine/Private/UserDefinedEnum.cpp
@@ -4,6 +4,7 @@
 #include "UObject/EditorObjectVersion.h"
 #include "UObject/ObjectSaveContext.h"
 #include "CookedMetaData.h"
+#include "UObject/Package.h"
 
 #include UE_INLINE_GENERATED_CPP_BY_NAME(UserDefinedEnum)
 
diff --git a/Engine/Source/Runtime/Engine/Private/WorldPartition/ContentBundle/ContentBundleWorldSubsystem.cpp b/Engine/Source/Runtime/Engine/Private/WorldPartition/ContentBundle/ContentBundleWorldSubsystem.cpp
--- a/Engine/Source/Runtime/Engine/Private/WorldPartition/ContentBundle/ContentBundleWorldSubsystem.cpp
+++ b/Engine/Source/Runtime/Engine/Private/WorldPartition/ContentBundle/ContentBundleWorldSubsystem.cpp
@@ -1,6 +1,7 @@
 // Copyright Epic Games, Inc. All Rights Reserved.
 #include "WorldPartition/ContentBundle/ContentBundleWorldSubsystem.h"
 
+#include "Engine/Engine.h"
 #include "WorldPartition/WorldPartition.h"
 #include "WorldPartition/ContentBundle/ContentBundle.h"
 #include "WorldPartition/ContentBundle/ContentBundleDescriptor.h"
diff --git a/Engine/Source/Runtime/Engine/Private/WorldPartition/WorldPartitionRuntimeSpatialHash.cpp b/Engine/Source/Runtime/Engine/Private/WorldPartition/WorldPartitionRuntimeSpatialHash.cpp
--- a/Engine/Source/Runtime/Engine/Private/WorldPartition/WorldPartitionRuntimeSpatialHash.cpp
+++ b/Engine/Source/Runtime/Engine/Private/WorldPartition/WorldPartitionRuntimeSpatialHash.cpp
@@ -5,6 +5,7 @@
 =============================================================================*/
 #include "WorldPartition/WorldPartitionRuntimeSpatialHash.h"
 
+#include "Engine/Engine.h"
 #include "WorldPartition/WorldPartitionStreamingGenerationContext.h"
 #include "WorldPartition/WorldPartitionRuntimeSpatialHashCell.h"
 #include "WorldPartition/WorldPartitionRuntimeLevelStreamingCell.h"
diff --git a/Engine/Source/Runtime/Engine/Private/WorldPartition/WorldPartitionStreamingPolicy.cpp b/Engine/Source/Runtime/Engine/Private/WorldPartition/WorldPartitionStreamingPolicy.cpp
--- a/Engine/Source/Runtime/Engine/Private/WorldPartition/WorldPartitionStreamingPolicy.cpp
+++ b/Engine/Source/Runtime/Engine/Private/WorldPartition/WorldPartitionStreamingPolicy.cpp
@@ -5,6 +5,7 @@
  */
 
 #include "WorldPartition/WorldPartitionStreamingPolicy.h"
+#include "Algo/Find.h"
 #include "WorldPartition/WorldPartitionRuntimeCell.h"
 #include "WorldPartition/WorldPartitionStreamingSource.h"
 #include "WorldPartition/WorldPartitionSubsystem.h"
diff --git a/Engine/Source/Runtime/Experimental/Iris/Core/Private/Iris/ReplicationSystem/ObjectReplicationBridge.cpp b/Engine/Source/Runtime/Experimental/Iris/Core/Private/Iris/ReplicationSystem/ObjectReplicationBridge.cpp
--- a/Engine/Source/Runtime/Experimental/Iris/Core/Private/Iris/ReplicationSystem/ObjectReplicationBridge.cpp
+++ b/Engine/Source/Runtime/Experimental/Iris/Core/Private/Iris/ReplicationSystem/ObjectReplicationBridge.cpp
@@ -2,6 +2,7 @@
 
 #include "Iris/ReplicationSystem/ObjectReplicationBridge.h"
 
+#include "HAL/IConsoleManager.h"
 #include "Iris/IrisConfigInternal.h"
 
 #include "Iris/Core/IrisLog.h"
