Fixed non-unity/pch errors

diff --git a/Engine/Plugins/Experimental/CommonUI/Source/CommonUI/Private/CommonGameViewportClient.cpp b/Engine/Plugins/Experimental/CommonUI/Source/CommonUI/Private/CommonGameViewportClient.cpp
--- a/Engine/Plugins/Experimental/CommonUI/Source/CommonUI/Private/CommonGameViewportClient.cpp
+++ b/Engine/Plugins/Experimental/CommonUI/Source/CommonUI/Private/CommonGameViewportClient.cpp
@@ -3,6 +3,7 @@
 #include "CommonGameViewportClient.h"
 #include "Engine/Console.h"
 #include "Engine/Engine.h"
+#include "Engine/GameInstance.h"
 #include "Engine/LocalPlayer.h"
 #include "GameFramework/GameUserSettings.h"
 #include "GameFramework/PlayerController.h"
diff --git a/Engine/Plugins/Experimental/GameplayTargetingSystem/Source/GameplayTargetingSystem/TargetingSystem/TargetingSubsystem.cpp b/Engine/Plugins/Experimental/GameplayTargetingSystem/Source/GameplayTargetingSystem/TargetingSystem/TargetingSubsystem.cpp
--- a/Engine/Plugins/Experimental/GameplayTargetingSystem/Source/GameplayTargetingSystem/TargetingSystem/TargetingSubsystem.cpp
+++ b/Engine/Plugins/Experimental/GameplayTargetingSystem/Source/GameplayTargetingSystem/TargetingSystem/TargetingSubsystem.cpp
@@ -1,6 +1,7 @@
 // Copyright Epic Games, Inc. All Rights Reserved.
 #include "TargetingSubsystem.h"
 
+#include "Engine/GameInstance.h"
 #include "Engine/World.h"
 #include "GameplayTargetingSystem/Types/TargetingSystemLogs.h"
 #include "GameplayTargetingSystem/Types/TargetingSystemTypes.h"
diff --git a/Engine/Plugins/FX/Niagara/Source/Niagara/Private/NiagaraPerfBaseline.cpp b/Engine/Plugins/FX/Niagara/Source/Niagara/Private/NiagaraPerfBaseline.cpp
--- a/Engine/Plugins/FX/Niagara/Source/Niagara/Private/NiagaraPerfBaseline.cpp
+++ b/Engine/Plugins/FX/Niagara/Source/Niagara/Private/NiagaraPerfBaseline.cpp
@@ -13,6 +13,7 @@
 #include "HighResScreenshot.h"
 #include "CanvasTypes.h"
 #include "HAL/FileManager.h"
+#include "UObject/UObjectIterator.h"
 
 #include UE_INLINE_GENERATED_CPP_BY_NAME(NiagaraPerfBaseline)
 
diff --git a/Engine/Plugins/MovieScene/MovieRenderPipeline/Source/MovieRenderPipelineCore/Private/MovieRenderPipelineCommandLine.cpp b/Engine/Plugins/MovieScene/MovieRenderPipeline/Source/MovieRenderPipelineCore/Private/MovieRenderPipelineCommandLine.cpp
--- a/Engine/Plugins/MovieScene/MovieRenderPipeline/Source/MovieRenderPipelineCore/Private/MovieRenderPipelineCommandLine.cpp
+++ b/Engine/Plugins/MovieScene/MovieRenderPipeline/Source/MovieRenderPipelineCore/Private/MovieRenderPipelineCommandLine.cpp
@@ -13,6 +13,7 @@
 #include "MoviePipelineQueue.h"
 #include "Misc/PackageName.h"
 #include "Misc/FileHelper.h"
+#include "UObject/Package.h"
 #include "UObject/UObjectHash.h"
 #include "MoviePipelinePythonHostExecutor.h"
 #include "MoviePipelineUtils.h"
diff --git a/Engine/Plugins/VirtualProduction/RemoteControl/Source/RemoteControl/Private/RemoteControlPreset.cpp b/Engine/Plugins/VirtualProduction/RemoteControl/Source/RemoteControl/Private/RemoteControlPreset.cpp
--- a/Engine/Plugins/VirtualProduction/RemoteControl/Source/RemoteControl/Private/RemoteControlPreset.cpp
+++ b/Engine/Plugins/VirtualProduction/RemoteControl/Source/RemoteControl/Private/RemoteControlPreset.cpp
@@ -29,6 +29,7 @@
 
 #include "UObject/Object.h"
 #include "UObject/ObjectMacros.h"
+#include "UObject/Package.h"
 #if WITH_EDITOR
 #include "AnalyticsEventAttribute.h"
 #include "Editor.h"
diff --git a/Engine/Source/Developer/GameplayDebugger/Private/GameplayDebuggerLocalController.cpp b/Engine/Source/Developer/GameplayDebugger/Private/GameplayDebuggerLocalController.cpp
--- a/Engine/Source/Developer/GameplayDebugger/Private/GameplayDebuggerLocalController.cpp
+++ b/Engine/Source/Developer/GameplayDebugger/Private/GameplayDebuggerLocalController.cpp
@@ -1,9 +1,11 @@
 // Copyright Epic Games, Inc. All Rights Reserved.
 
 #include "GameplayDebuggerLocalController.h"
+#include "Engine/Engine.h"
 #include "InputCoreTypes.h"
 #include "Framework/Commands/InputChord.h"
 #include "Components/InputComponent.h"
+#include "SceneView.h"
 #include "TimerManager.h"
 #include "GameFramework/Pawn.h"
 #include "GameFramework/Controller.h"
diff --git a/Engine/Source/Runtime/Engine/Private/ActorConstruction.cpp b/Engine/Source/Runtime/Engine/Private/ActorConstruction.cpp
--- a/Engine/Source/Runtime/Engine/Private/ActorConstruction.cpp
+++ b/Engine/Source/Runtime/Engine/Private/ActorConstruction.cpp
@@ -3,6 +3,7 @@
 #include "CoreMinimal.h"
 #include "Math/RandomStream.h"
 #include "Stats/Stats.h"
+#include "UObject/Package.h"
 #include "UObject/Script.h"
 #include "UObject/ObjectMacros.h"
 #include "UObject/Object.h"
diff --git a/Engine/Source/Runtime/Engine/Private/WorldPartition/WorldPartitionRuntimeSpatialHash.cpp b/Engine/Source/Runtime/Engine/Private/WorldPartition/WorldPartitionRuntimeSpatialHash.cpp
--- a/Engine/Source/Runtime/Engine/Private/WorldPartition/WorldPartitionRuntimeSpatialHash.cpp
+++ b/Engine/Source/Runtime/Engine/Private/WorldPartition/WorldPartitionRuntimeSpatialHash.cpp
@@ -6,6 +6,7 @@
 #include "WorldPartition/WorldPartitionRuntimeSpatialHash.h"
 
 #include "Engine/Engine.h"
+#include "Misc/Paths.h"
 #include "WorldPartition/WorldPartitionStreamingGenerationContext.h"
 #include "WorldPartition/WorldPartitionRuntimeSpatialHashCell.h"
 #include "WorldPartition/WorldPartitionRuntimeLevelStreamingCell.h"
