Fixed non-unity/pch errors

diff --git a/Engine/Plugins/Experimental/GameplayTargetingSystem/Source/GameplayTargetingSystem/TargetingSystem/TargetingSubsystem.cpp b/Engine/Plugins/Experimental/GameplayTargetingSystem/Source/GameplayTargetingSystem/TargetingSystem/TargetingSubsystem.cpp
--- a/Engine/Plugins/Experimental/GameplayTargetingSystem/Source/GameplayTargetingSystem/TargetingSystem/TargetingSubsystem.cpp
+++ b/Engine/Plugins/Experimental/GameplayTargetingSystem/Source/GameplayTargetingSystem/TargetingSystem/TargetingSubsystem.cpp
@@ -1,6 +1,7 @@
 // Copyright Epic Games, Inc. All Rights Reserved.
 #include "TargetingSubsystem.h"
 
+#include "Engine/Engine.h"
 #include "Engine/GameInstance.h"
 #include "Engine/World.h"
 #include "GameplayTargetingSystem/Types/TargetingSystemLogs.h"
diff --git a/Engine/Plugins/Media/ElectraPlayer/Source/ElectraPlayerRuntime/Private/Runtime/Decoder/Android/AudioDecoderAAC_Android.cpp b/Engine/Plugins/Media/ElectraPlayer/Source/ElectraPlayerRuntime/Private/Runtime/Decoder/Android/AudioDecoderAAC_Android.cpp
--- a/Engine/Plugins/Media/ElectraPlayer/Source/ElectraPlayerRuntime/Private/Runtime/Decoder/Android/AudioDecoderAAC_Android.cpp
+++ b/Engine/Plugins/Media/ElectraPlayer/Source/ElectraPlayerRuntime/Private/Runtime/Decoder/Android/AudioDecoderAAC_Android.cpp
@@ -3,6 +3,7 @@
 #include "PlayerCore.h"
 #include "PlayerRuntimeGlobal.h"
 
+#include "Stats/Stats.h"
 #include "StreamAccessUnitBuffer.h"
 #include "Decoder/AudioDecoderAAC.h"
 #include "Renderer/RendererBase.h"
diff --git a/Engine/Plugins/Media/ElectraPlayer/Source/ElectraPlayerRuntime/Private/Runtime/Decoder/Android/VideoDecoderH264_Android.cpp b/Engine/Plugins/Media/ElectraPlayer/Source/ElectraPlayerRuntime/Private/Runtime/Decoder/Android/VideoDecoderH264_Android.cpp
--- a/Engine/Plugins/Media/ElectraPlayer/Source/ElectraPlayerRuntime/Private/Runtime/Decoder/Android/VideoDecoderH264_Android.cpp
+++ b/Engine/Plugins/Media/ElectraPlayer/Source/ElectraPlayerRuntime/Private/Runtime/Decoder/Android/VideoDecoderH264_Android.cpp
@@ -3,6 +3,7 @@
 #include "PlayerCore.h"
 #include "PlayerRuntimeGlobal.h"
 
+#include "Stats/Stats.h"
 #include "StreamAccessUnitBuffer.h"
 #include "Decoder/VideoDecoderH264.h"
 #include "Decoder/Android/DecoderOptionNames_Android.h"
diff --git a/Engine/Plugins/Media/ElectraPlayer/Source/ElectraPlayerRuntime/Private/Runtime/Decoder/Android/VideoDecoderH265_Android.cpp b/Engine/Plugins/Media/ElectraPlayer/Source/ElectraPlayerRuntime/Private/Runtime/Decoder/Android/VideoDecoderH265_Android.cpp
--- a/Engine/Plugins/Media/ElectraPlayer/Source/ElectraPlayerRuntime/Private/Runtime/Decoder/Android/VideoDecoderH265_Android.cpp
+++ b/Engine/Plugins/Media/ElectraPlayer/Source/ElectraPlayerRuntime/Private/Runtime/Decoder/Android/VideoDecoderH265_Android.cpp
@@ -3,6 +3,7 @@
 #include "PlayerCore.h"
 #include "PlayerRuntimeGlobal.h"
 
+#include "Stats/Stats.h"
 #include "StreamAccessUnitBuffer.h"
 #include "Decoder/VideoDecoderH265.h"
 #include "Decoder/Android/DecoderOptionNames_Android.h"
diff --git a/Engine/Plugins/Media/ElectraPlayer/Source/ElectraPlayerRuntime/Private/Runtime/Decoder/Apple/AudioDecoderAAC_Apple.cpp b/Engine/Plugins/Media/ElectraPlayer/Source/ElectraPlayerRuntime/Private/Runtime/Decoder/Apple/AudioDecoderAAC_Apple.cpp
--- a/Engine/Plugins/Media/ElectraPlayer/Source/ElectraPlayerRuntime/Private/Runtime/Decoder/Apple/AudioDecoderAAC_Apple.cpp
+++ b/Engine/Plugins/Media/ElectraPlayer/Source/ElectraPlayerRuntime/Private/Runtime/Decoder/Apple/AudioDecoderAAC_Apple.cpp
@@ -7,6 +7,7 @@
 #include "PlayerCore.h"
 #include "PlayerRuntimeGlobal.h"
 
+#include "Stats/Stats.h"
 #include "StreamAccessUnitBuffer.h"
 #include "Decoder/AudioDecoderAAC.h"
 #include "Renderer/RendererBase.h"
diff --git a/Engine/Plugins/Media/ElectraPlayer/Source/ElectraPlayerRuntime/Private/Runtime/Decoder/Apple/VideoDecoderH264_Apple.cpp b/Engine/Plugins/Media/ElectraPlayer/Source/ElectraPlayerRuntime/Private/Runtime/Decoder/Apple/VideoDecoderH264_Apple.cpp
--- a/Engine/Plugins/Media/ElectraPlayer/Source/ElectraPlayerRuntime/Private/Runtime/Decoder/Apple/VideoDecoderH264_Apple.cpp
+++ b/Engine/Plugins/Media/ElectraPlayer/Source/ElectraPlayerRuntime/Private/Runtime/Decoder/Apple/VideoDecoderH264_Apple.cpp
@@ -7,6 +7,7 @@
 #include "PlayerCore.h"
 #include "PlayerRuntimeGlobal.h"
 
+#include "Stats/Stats.h"
 #include "StreamAccessUnitBuffer.h"
 #include "Decoder/VideoDecoderH264.h"
 #include "Renderer/RendererBase.h"
diff --git a/Engine/Plugins/VirtualProduction/Takes/Source/TakeRecorderSources/Private/TakeRecorderCameraCutSource.cpp b/Engine/Plugins/VirtualProduction/Takes/Source/TakeRecorderSources/Private/TakeRecorderCameraCutSource.cpp
--- a/Engine/Plugins/VirtualProduction/Takes/Source/TakeRecorderSources/Private/TakeRecorderCameraCutSource.cpp
+++ b/Engine/Plugins/VirtualProduction/Takes/Source/TakeRecorderSources/Private/TakeRecorderCameraCutSource.cpp
@@ -2,6 +2,7 @@
 
 #include "TakeRecorderCameraCutSource.h"
 
+#include "AssetRegistry/AssetData.h"
 #include "TakeRecorderSources.h"
 #include "TakeRecorderSource.h"
 #include "TakeRecorderActorSource.h"
diff --git a/Engine/Plugins/VirtualProduction/Takes/Source/TakeRecorderSources/Private/TakeRecorderLevelSequenceSource.cpp b/Engine/Plugins/VirtualProduction/Takes/Source/TakeRecorderSources/Private/TakeRecorderLevelSequenceSource.cpp
--- a/Engine/Plugins/VirtualProduction/Takes/Source/TakeRecorderSources/Private/TakeRecorderLevelSequenceSource.cpp
+++ b/Engine/Plugins/VirtualProduction/Takes/Source/TakeRecorderSources/Private/TakeRecorderLevelSequenceSource.cpp
@@ -1,6 +1,7 @@
 // Copyright Epic Games, Inc. All Rights Reserved.
 
 #include "TakeRecorderLevelSequenceSource.h"
+#include "AssetRegistry/AssetData.h"
 #include "TakesCoreLog.h"
 #include "TakesUtils.h"
 #include "Editor.h"
diff --git a/Engine/Plugins/VirtualProduction/Takes/Source/TakeRecorderSources/Private/TakeRecorderLevelVisibilitySource.cpp b/Engine/Plugins/VirtualProduction/Takes/Source/TakeRecorderSources/Private/TakeRecorderLevelVisibilitySource.cpp
--- a/Engine/Plugins/VirtualProduction/Takes/Source/TakeRecorderSources/Private/TakeRecorderLevelVisibilitySource.cpp
+++ b/Engine/Plugins/VirtualProduction/Takes/Source/TakeRecorderSources/Private/TakeRecorderLevelVisibilitySource.cpp
@@ -1,6 +1,7 @@
 // Copyright Epic Games, Inc. All Rights Reserved.
 
 #include "TakeRecorderLevelVisibilitySource.h"
+#include "AssetRegistry/AssetData.h"
 #include "TakesUtils.h"
 
 #include "TakeRecorderSources.h"
diff --git a/Engine/Plugins/VirtualProduction/Takes/Source/TakeRecorderSources/Private/TakeRecorderMicrophoneAudioSource.cpp b/Engine/Plugins/VirtualProduction/Takes/Source/TakeRecorderSources/Private/TakeRecorderMicrophoneAudioSource.cpp
--- a/Engine/Plugins/VirtualProduction/Takes/Source/TakeRecorderSources/Private/TakeRecorderMicrophoneAudioSource.cpp
+++ b/Engine/Plugins/VirtualProduction/Takes/Source/TakeRecorderSources/Private/TakeRecorderMicrophoneAudioSource.cpp
@@ -1,6 +1,7 @@
 // Copyright Epic Games, Inc. All Rights Reserved.
 
 #include "TakeRecorderMicrophoneAudioSource.h"
+#include "AssetRegistry/AssetData.h"
 #include "TakeRecorderSources.h"
 #include "TakeRecorderSettings.h"
 #include "TakesUtils.h"
diff --git a/Engine/Plugins/VirtualProduction/Takes/Source/TakeRecorderSources/Private/TakeRecorderPlayerSource.cpp b/Engine/Plugins/VirtualProduction/Takes/Source/TakeRecorderSources/Private/TakeRecorderPlayerSource.cpp
--- a/Engine/Plugins/VirtualProduction/Takes/Source/TakeRecorderSources/Private/TakeRecorderPlayerSource.cpp
+++ b/Engine/Plugins/VirtualProduction/Takes/Source/TakeRecorderSources/Private/TakeRecorderPlayerSource.cpp
@@ -1,6 +1,7 @@
 // Copyright Epic Games, Inc. All Rights Reserved.
 
 #include "TakeRecorderPlayerSource.h"
+#include "AssetRegistry/AssetData.h"
 #include "GameFramework/Pawn.h"
 #include "TakesUtils.h"
 
diff --git a/Engine/Plugins/VirtualProduction/Takes/Source/TakeRecorderSources/Private/TakeRecorderWorldSource.cpp b/Engine/Plugins/VirtualProduction/Takes/Source/TakeRecorderSources/Private/TakeRecorderWorldSource.cpp
--- a/Engine/Plugins/VirtualProduction/Takes/Source/TakeRecorderSources/Private/TakeRecorderWorldSource.cpp
+++ b/Engine/Plugins/VirtualProduction/Takes/Source/TakeRecorderSources/Private/TakeRecorderWorldSource.cpp
@@ -2,6 +2,7 @@
 
 #include "TakeRecorderWorldSource.h"
 
+#include "AssetRegistry/AssetData.h"
 #include "TakeRecorderSources.h"
 #include "TakeRecorderActorSource.h"
 #include "TakeRecorderSourcesUtils.h"
diff --git a/Engine/Source/Runtime/Core/Private/Android/AndroidStats.cpp b/Engine/Source/Runtime/Core/Private/Android/AndroidStats.cpp
--- a/Engine/Source/Runtime/Core/Private/Android/AndroidStats.cpp
+++ b/Engine/Source/Runtime/Core/Private/Android/AndroidStats.cpp
@@ -1,6 +1,7 @@
 // Copyright Epic Games, Inc. All Rights Reserved.
 
 #include "Android/AndroidStats.h"
+#include "Async/TaskGraphInterfaces.h"
 #include "CoreMinimal.h"
 #include "Android/AndroidPlatformMisc.h"
 #include "Stats/Stats.h"
diff --git a/Engine/Source/Runtime/Core/Private/HAL/Allocators/CachedOSVeryLargePageAllocator.cpp b/Engine/Source/Runtime/Core/Private/HAL/Allocators/CachedOSVeryLargePageAllocator.cpp
--- a/Engine/Source/Runtime/Core/Private/HAL/Allocators/CachedOSVeryLargePageAllocator.cpp
+++ b/Engine/Source/Runtime/Core/Private/HAL/Allocators/CachedOSVeryLargePageAllocator.cpp
@@ -1,6 +1,7 @@
 // Copyright Epic Games, Inc. All Rights Reserved.
 
 #include "HAL/Allocators/CachedOSVeryLargePageAllocator.h"
+#include "HAL/IConsoleManager.h"
 #include "HAL/UnrealMemory.h"
 #include "Logging/LogMacros.h"
 #include "CoreGlobals.h"
diff --git a/Engine/Source/Runtime/Engine/Private/Streaming/Texture2DStreamIn_IO.cpp b/Engine/Source/Runtime/Engine/Private/Streaming/Texture2DStreamIn_IO.cpp
--- a/Engine/Source/Runtime/Engine/Private/Streaming/Texture2DStreamIn_IO.cpp
+++ b/Engine/Source/Runtime/Engine/Private/Streaming/Texture2DStreamIn_IO.cpp
@@ -5,6 +5,7 @@ Texture2DStreamIn.cpp: Stream in helper for 2D textures using texture streaming
 =============================================================================*/
 
 #include "Streaming/Texture2DStreamIn_IO.h"
+#include "EngineLogs.h"
 #include "RenderUtils.h"
 #include "HAL/PlatformFileManager.h"
 #include "HAL/FileManager.h"
