GameplayTargetingSystem: Add mechanism for cancelling async TargetingRequests

diff --git a/Engine/Plugins/Experimental/GameplayTargetingSystem/Source/GameplayTargetingSystem/TargetingSystem/TargetingSubsystem.cpp b/Engine/Plugins/Experimental/GameplayTargetingSystem/Source/GameplayTargetingSystem/TargetingSystem/TargetingSubsystem.cpp
--- a/Engine/Plugins/Experimental/GameplayTargetingSystem/Source/GameplayTargetingSystem/TargetingSystem/TargetingSubsystem.cpp
+++ b/Engine/Plugins/Experimental/GameplayTargetingSystem/Source/GameplayTargetingSystem/TargetingSystem/TargetingSubsystem.cpp
@@ -360,16 +360,21 @@ void UTargetingSubsystem::RemoveAsyncTargetingRequestWithHandle(FTargetingReques
 {
 	if (TargetingHandle.IsValid())
 	{
+		if (UTargetingTask* ExecutingTask = FindCurrentExecutingTask(TargetingHandle))
+		{
+			ExecutingTask->CancelAsync();
+		}
+
 		AsyncTargetingRequests.RemoveAll([TargetingHandle](const FTargetingRequestHandle& Handle)
-			{
-				return (TargetingHandle == Handle);
-			});
+		{
+			return (TargetingHandle == Handle);
+		});
 
 		UTargetingSubsystem::ReleaseTargetRequestHandle(TargetingHandle);
 	}
 }
 
-void UTargetingSubsystem::StartAsyncTargetingRequest(const UTargetingPreset* TargetingPreset, const FTargetingSourceContext& SourceContext, FTargetingRequestDynamicDelegate CompletionDynamicDelegate)
+FTargetingRequestHandle UTargetingSubsystem::StartAsyncTargetingRequest(const UTargetingPreset* TargetingPreset, const FTargetingSourceContext& SourceContext, FTargetingRequestDynamicDelegate CompletionDynamicDelegate)
 {
 	FTargetingRequestHandle RequestHandle;
 	if (ensure(TargetingPreset))
@@ -389,6 +394,8 @@ void UTargetingSubsystem::StartAsyncTargetingRequest(const UTargetingPreset* Tar
 		FTargetingRequestDelegate CompletionDelegate;
 		StartAsyncTargetingRequestWithHandleInternal(RequestHandle, CompletionDelegate, CompletionDynamicDelegate);
 	}
+	
+	return RequestHandle;
 }
 
 void UTargetingSubsystem::ExecuteTargetingRequestWithHandleInternal(FTargetingRequestHandle TargetingHandle, FTargetingRequestDelegate CompletionDelegate, FTargetingRequestDynamicDelegate CompletionDynamicDelegate) const
@@ -625,6 +632,33 @@ void UTargetingSubsystem::OnShowDebugInfo(AHUD* HUD, UCanvas* Canvas, const FDeb
 	}
 }
 
+UTargetingTask* UTargetingSubsystem::FindCurrentExecutingTask(FTargetingRequestHandle Handle) const
+{
+	if (Handle.IsValid())
+	{
+		if (FTargetingAsyncTaskData* TaskData = FTargetingAsyncTaskData::Find(Handle))
+		{
+			const int32 CurrentTaskIndex = TaskData->CurrentAsyncTaskIndex;
+
+			if (TaskData->CurrentAsyncTaskState == ETargetingTaskAsyncState::Executing)
+			{
+				if (const FTargetingTaskSet** PtrToTaskSet = FTargetingTaskSet::Find(Handle))
+				{
+					if (const FTargetingTaskSet* TaskSet = (*PtrToTaskSet))
+					{
+						if (ensure(TaskSet->Tasks.IsValidIndex(CurrentTaskIndex)))
+						{
+							return TaskSet->Tasks[CurrentTaskIndex];
+						}
+					}
+				}
+			}
+		}
+	}
+
+	return nullptr;
+}
+
 void UTargetingSubsystem::DisplayDebug(class UCanvas* Canvas, const class FDebugDisplayInfo& DebugDisplay, float& YL, float& YPos)
 {
 	FTargetingDebugInfo DebugInfo;
diff --git a/Engine/Plugins/Experimental/GameplayTargetingSystem/Source/GameplayTargetingSystem/TargetingSystem/TargetingSubsystem.h b/Engine/Plugins/Experimental/GameplayTargetingSystem/Source/GameplayTargetingSystem/TargetingSystem/TargetingSubsystem.h
--- a/Engine/Plugins/Experimental/GameplayTargetingSystem/Source/GameplayTargetingSystem/TargetingSystem/TargetingSubsystem.h
+++ b/Engine/Plugins/Experimental/GameplayTargetingSystem/Source/GameplayTargetingSystem/TargetingSystem/TargetingSubsystem.h
@@ -140,11 +140,12 @@ public:
 	TARGETINGSYSTEM_API void StartAsyncTargetingRequestWithHandle(FTargetingRequestHandle TargetingHandle, FTargetingRequestDelegate CompletionDelegate);
 
 	/** Method to remove an async targeting request with a given targeting handle */
+	UFUNCTION(BlueprintCallable, Category = "Gameplay Targeting | Async Request")
 	TARGETINGSYSTEM_API void RemoveAsyncTargetingRequestWithHandle(FTargetingRequestHandle& TargetingHandle);
 
 	/** Method to queue an async targeting request based on a gameplay targeting preset. */
-	UFUNCTION(BlueprintCallable, BlueprintPure = false, Category = "Gameplay Targeting | Async Request ")
-	void StartAsyncTargetingRequest(const UTargetingPreset* TargetingPreset, const FTargetingSourceContext& InSourceContext, FTargetingRequestDynamicDelegate CompletionDynamicDelegate);
+	UFUNCTION(BlueprintCallable, Category = "Gameplay Targeting | Async Request")
+	FTargetingRequestHandle StartAsyncTargetingRequest(const UTargetingPreset* TargetingPreset, const FTargetingSourceContext& InSourceContext, FTargetingRequestDynamicDelegate CompletionDynamicDelegate);
 
 private:
 	/** Internal method to handle running an immediate targeting request */
@@ -160,6 +161,7 @@ private:
 	UPROPERTY(Transient)
 	TArray<FTargetingRequestHandle> AsyncTargetingRequests;
 
+	UTargetingTask* FindCurrentExecutingTask(FTargetingRequestHandle Handle) const;
 	/** ~Targeting Request Methods */
 
 	/** Blueprint Helper Methods */
diff --git a/Engine/Plugins/Experimental/GameplayTargetingSystem/Source/GameplayTargetingSystem/Tasks/TargetingTask.cpp b/Engine/Plugins/Experimental/GameplayTargetingSystem/Source/GameplayTargetingSystem/Tasks/TargetingTask.cpp
--- a/Engine/Plugins/Experimental/GameplayTargetingSystem/Source/GameplayTargetingSystem/Tasks/TargetingTask.cpp
+++ b/Engine/Plugins/Experimental/GameplayTargetingSystem/Source/GameplayTargetingSystem/Tasks/TargetingTask.cpp
@@ -19,6 +19,10 @@ void UTargetingTask::Execute(const FTargetingRequestHandle& TargetingHandle) con
 {
 }
 
+void UTargetingTask::CancelAsync() const
+{
+}
+
 bool UTargetingTask::IsAsyncTargetingRequest(const FTargetingRequestHandle& TargetingHandle) const
 {
 	if (FTargetingAsyncTaskData* AsyncStateData = FTargetingAsyncTaskData::Find(TargetingHandle))
diff --git a/Engine/Plugins/Experimental/GameplayTargetingSystem/Source/GameplayTargetingSystem/Tasks/TargetingTask.h b/Engine/Plugins/Experimental/GameplayTargetingSystem/Source/GameplayTargetingSystem/Tasks/TargetingTask.h
--- a/Engine/Plugins/Experimental/GameplayTargetingSystem/Source/GameplayTargetingSystem/Tasks/TargetingTask.h
+++ b/Engine/Plugins/Experimental/GameplayTargetingSystem/Source/GameplayTargetingSystem/Tasks/TargetingTask.h
@@ -54,6 +54,9 @@ public:
 	/** Evaluation function called by derived classes to process the targeting request */
 	virtual void Execute(const FTargetingRequestHandle& TargetingHandle) const;
 
+	/** Lifecycle function called when the task was cancelled while in the Executing Async state */
+	virtual void CancelAsync() const;
+
 protected:
 	/** Helper method to check if this task is running in an async targeting request */
 	bool IsAsyncTargetingRequest(const FTargetingRequestHandle& TargetingHandle) const;
