Add an AbilityTask for performing targeting requests with the TargetingSystem plugin

diff --git a/Engine/Plugins/Experimental/GameplayTargetingSystem/Source/GameplayTargetingSystem/AbilityTasks/AbilityTask_PerformTargeting.cpp b/Engine/Plugins/Experimental/GameplayTargetingSystem/Source/GameplayTargetingSystem/AbilityTasks/AbilityTask_PerformTargeting.cpp
new file mode 100644
--- /dev/null
+++ b/Engine/Plugins/Experimental/GameplayTargetingSystem/Source/GameplayTargetingSystem/AbilityTasks/AbilityTask_PerformTargeting.cpp
@@ -0,0 +1,98 @@
+// Copyright Epic Games, Inc. All Rights Reserved.
+#include "AbilityTask_PerformTargeting.h"
+#include "GameplayTargetingSystem/TargetingSystem/TargetingSubsystem.h"
+#include "GameplayTargetingSystem/Types/TargetingSystemLogs.h"
+
+void UAbilityTask_PerformTargeting::Activate()
+{
+	if (AActor* SourceActor = GetAvatarActor())
+	{
+		if (UTargetingSubsystem* TargetingSubsystem = UTargetingSubsystem::Get(SourceActor->GetWorld()))
+		{
+			FTargetingSourceContext SourceContext;
+			SourceContext.SourceActor = SourceActor;
+
+			FTargetingRequestHandle TargetingHandle = UTargetingSubsystem::MakeTargetRequestHandle(TargetingPreset, SourceContext);
+			
+			SetupInitialTargetsForRequest(TargetingHandle);
+			
+			FTargetingRequestDelegate Delegate = FTargetingRequestDelegate::CreateWeakLambda(this, [this](FTargetingRequestHandle TargetingHandle)
+			{
+				OnTargetReady.Broadcast(TargetingHandle);
+			});
+
+			if (bPerformAsync)
+			{
+				TARGETING_LOG(Verbose, TEXT("Starting async targeting ability task"));
+
+				FTargetingAsyncTaskData& AsyncTaskData = FTargetingAsyncTaskData::FindOrAdd(TargetingHandle);
+				AsyncTaskData.bReleaseOnCompletion = true;
+
+				TargetingSubsystem->StartAsyncTargetingRequestWithHandle(TargetingHandle, Delegate);
+			}
+			else
+			{
+				TARGETING_LOG(Verbose, TEXT("Starting immediate targeting ability task"));
+
+				FTargetingImmediateTaskData& ImmeidateTaskData = FTargetingImmediateTaskData::FindOrAdd(TargetingHandle);
+				ImmeidateTaskData.bReleaseOnCompletion = true;
+
+				TargetingSubsystem->ExecuteTargetingRequestWithHandle(TargetingHandle, Delegate);
+
+				EndTask();
+			}
+		}
+	}
+	else
+	{
+		TARGETING_LOG(Error, TEXT("%s called with an invalid source actor!"), ANSI_TO_TCHAR(__FUNCTION__));
+		EndTask();
+	}
+}
+
+UAbilityTask_PerformTargeting* UAbilityTask_PerformTargeting::PerformTargetingRequest(UGameplayAbility* OwningAbility, UTargetingPreset* InTargetingPreset, bool bAllowAsync)
+{
+	UAbilityTask_PerformTargeting* Task = NewAbilityTask<UAbilityTask_PerformTargeting>(OwningAbility);
+	Task->TargetingPreset = InTargetingPreset;
+	Task->bPerformAsync = bAllowAsync;
+
+	return Task;
+}
+
+UAbilityTask_PerformTargeting* UAbilityTask_PerformTargeting::PerformFilteringRequest(UGameplayAbility* OwningAbility, UTargetingPreset* InTargetingPreset, const TArray<AActor*> InTargets, bool bAllowAsync)
+{
+	UAbilityTask_PerformTargeting* Task = NewAbilityTask<UAbilityTask_PerformTargeting>(OwningAbility);
+	Task->TargetingPreset = InTargetingPreset;
+	Task->InitialTargets = InTargets;
+	Task->bPerformAsync = bAllowAsync;
+
+	return Task;
+}
+
+void UAbilityTask_PerformTargeting::SetupInitialTargetsForRequest(FTargetingRequestHandle RequestHandle) const
+{
+	if (RequestHandle.IsValid() && InitialTargets.Num() > 0)
+	{
+		FTargetingDefaultResultsSet& TargetingResults = FTargetingDefaultResultsSet::FindOrAdd(RequestHandle);
+		for (AActor* Target : InitialTargets)
+		{
+			if (!Target)
+			{
+				continue;
+			}
+
+			bool bAddResult = !TargetingResults.TargetResults.ContainsByPredicate([Target](const FTargetingDefaultResultData& Data) -> bool
+			{
+				return (Data.HitResult.GetActor() == Target);
+			});
+
+			if (bAddResult)
+			{
+				FTargetingDefaultResultData* ResultData = new(TargetingResults.TargetResults) FTargetingDefaultResultData();
+				ResultData->HitResult.HitObjectHandle = FActorInstanceHandle(Target);
+				ResultData->HitResult.Location = Target->GetActorLocation();
+			}
+		}
+
+	}
+}
diff --git a/Engine/Plugins/Experimental/GameplayTargetingSystem/Source/GameplayTargetingSystem/AbilityTasks/AbilityTask_PerformTargeting.h b/Engine/Plugins/Experimental/GameplayTargetingSystem/Source/GameplayTargetingSystem/AbilityTasks/AbilityTask_PerformTargeting.h
new file mode 100644
--- /dev/null
+++ b/Engine/Plugins/Experimental/GameplayTargetingSystem/Source/GameplayTargetingSystem/AbilityTasks/AbilityTask_PerformTargeting.h
@@ -0,0 +1,47 @@
+// Copyright Epic Games, Inc. All Rights Reserved.
+#pragma once
+
+#include "CoreMinimal.h"
+#include "UObject/ObjectMacros.h"
+#include "Engine/EngineTypes.h"
+#include "Abilities/Tasks/AbilityTask.h"
+#include "GameplayTargetingSystem/Types/TargetingSystemTypes.h"
+#include "AbilityTask_PerformTargeting.generated.h"
+
+class UTargetingPreset;
+
+DECLARE_DYNAMIC_MULTICAST_DELEGATE_OneParam(FTargetReadyDelegate, FTargetingRequestHandle, TargetingRequestHandle);
+
+UCLASS()
+class UAbilityTask_PerformTargeting : public UAbilityTask
+{
+	GENERATED_BODY()
+
+	// Called when the targeting request has been completed and results are ready
+	UPROPERTY(BlueprintAssignable)
+	FTargetReadyDelegate OnTargetReady;
+
+	virtual void Activate() override;
+
+	// Performs a targeting request based on a Targeting Preset from a GameplayAbility
+	UFUNCTION(BlueprintCallable, Category = "Ability|Tasks", meta = (HidePin = "OwningAbility", DefaultToSelf = "OwningAbility", BlueprintInternalUseOnly = "TRUE"))
+	static UAbilityTask_PerformTargeting* PerformTargetingRequest(UGameplayAbility* OwningAbility, UTargetingPreset* InTargetingPreset, bool bAllowAsync);
+	
+	// Performs a target filtering request based on a Targeting Preset from a GameplayAbility
+	UFUNCTION(BlueprintCallable, Category = "Ability|Tasks", meta = (HidePin = "OwningAbility", DefaultToSelf = "OwningAbility", BlueprintInternalUseOnly = "TRUE"))
+	static UAbilityTask_PerformTargeting* PerformFilteringRequest(UGameplayAbility* OwningAbility, UTargetingPreset* TargetingPreset, const TArray<AActor*> InTargets, bool bAllowAsync);
+
+private:
+	// Helper function to push the InitialTargets array to the DataStore for Filtering requests
+	void SetupInitialTargetsForRequest(FTargetingRequestHandle RequestHandle) const;
+
+protected:
+	// Targeting Preset to request
+	TObjectPtr<UTargetingPreset> TargetingPreset;
+
+	// Initial Targets to pass into the request (only populated for Filtering Requests)
+	TArray<AActor*> InitialTargets;
+
+	// Whether to perform this request immediately or async
+	bool bPerformAsync;
+};
\ No newline at end of file
diff --git a/Engine/Plugins/Experimental/GameplayTargetingSystem/Source/GameplayTargetingSystem/TargetingSystem.Build.cs b/Engine/Plugins/Experimental/GameplayTargetingSystem/Source/GameplayTargetingSystem/TargetingSystem.Build.cs
--- a/Engine/Plugins/Experimental/GameplayTargetingSystem/Source/GameplayTargetingSystem/TargetingSystem.Build.cs
+++ b/Engine/Plugins/Experimental/GameplayTargetingSystem/Source/GameplayTargetingSystem/TargetingSystem.Build.cs
@@ -21,6 +21,7 @@ public class TargetingSystem : ModuleRules
 				"Core",
 				"CoreUObject",
 				"Engine",
+				"GameplayTasks",
 				"GameplayAbilities",
 				"GameplayTags",
 			}
