GameplayTargetingSystem changes moved in to Engine plugin

diff --git a/Engine/Plugins/Experimental/GameplayTargetingSystem/Source/GameplayTargetingSystem/TargetingSystem/TargetingSubsystem.h b/Engine/Plugins/Experimental/GameplayTargetingSystem/Source/GameplayTargetingSystem/TargetingSystem/TargetingSubsystem.h
--- a/Engine/Plugins/Experimental/GameplayTargetingSystem/Source/GameplayTargetingSystem/TargetingSystem/TargetingSubsystem.h
+++ b/Engine/Plugins/Experimental/GameplayTargetingSystem/Source/GameplayTargetingSystem/TargetingSystem/TargetingSubsystem.h
@@ -172,7 +172,7 @@ public:
 
 	/** Method to get the actor targets from a given targeting request handle */
 	UFUNCTION(BlueprintCallable, Category = "Targeting System | Targeting Results")
-	static void GetTargetingResultsActors(FTargetingRequestHandle TargetingHandle, TArray<AActor*>& Targets);
+	static TARGETINGSYSTEM_API void GetTargetingResultsActors(FTargetingRequestHandle TargetingHandle, TArray<AActor*>& Targets);
 
 	/** Helper method to get the set of hit results for a given targeting handle */
 	UFUNCTION(BlueprintCallable, Category = "Targeting System | Targeting Results")
diff --git a/Engine/Plugins/Experimental/GameplayTargetingSystem/Source/GameplayTargetingSystem/Tasks/TargetingFilterTask_SortByDistance.cpp b/Engine/Plugins/Experimental/GameplayTargetingSystem/Source/GameplayTargetingSystem/Tasks/TargetingFilterTask_SortByDistance.cpp
--- a/Engine/Plugins/Experimental/GameplayTargetingSystem/Source/GameplayTargetingSystem/Tasks/TargetingFilterTask_SortByDistance.cpp
+++ b/Engine/Plugins/Experimental/GameplayTargetingSystem/Source/GameplayTargetingSystem/Tasks/TargetingFilterTask_SortByDistance.cpp
@@ -1,182 +1,36 @@
 // Copyright Epic Games, Inc. All Rights Reserved.
+
 #include "TargetingFilterTask_SortByDistance.h"
 
 #include "GameFramework/Actor.h"
-#include "GameplayTargetingSystem/TargetingSystem/TargetingSubsystem.h"
-
-#if ENABLE_DRAW_DEBUG
-#include "Engine/Canvas.h"
-#endif // ENABLE_DRAW_DEBUG
 
-namespace SortByDistanceConstants
-{
-	const FString PreSortPrefix = TEXT("PreSort");
-	const FString PostSortPrefix = TEXT("PostSort");
-}
+#include UE_INLINE_GENERATED_CPP_BY_NAME(TargetingFilterTask_SortByDistance)
 
-UTargetingFilterTask_SortByDistance::UTargetingFilterTask_SortByDistance(const FObjectInitializer& ObjectInitializer)
-: Super(ObjectInitializer)
+float UTargetingFilterTask_SortByDistance::GetScoreForTarget(const FTargetingRequestHandle& TargetingHandle, const FTargetingDefaultResultData& TargetData) const
 {
-	bAscending = true;
-}
-
-void UTargetingFilterTask_SortByDistance::Execute(const FTargetingRequestHandle& TargetingHandle) const
-{
-	Super::Execute(TargetingHandle);
-
-	SetTaskAsyncState(TargetingHandle, ETargetingTaskAsyncState::Executing);
-
-#if ENABLE_DRAW_DEBUG
-	ResetSortDebugStrings(TargetingHandle);
-#endif // ENABLE_DRAW_DEBUG
-
-	if (TargetingHandle.IsValid())
+	FVector SourceLocation = GetSourceLocation(TargetingHandle);
+	if (AActor* TargetActor = TargetData.HitResult.GetActor())
 	{
-		FVector SourceLocation = FVector::ZeroVector;
-		if (FTargetingSourceContext* SourceContext = FTargetingSourceContext::Find(TargetingHandle))
-		{
-			if (SourceContext->SourceActor)
-			{
-				SourceLocation = SourceContext->SourceActor->GetActorLocation();
-			}
-			else if (!SourceContext->SourceLocation.IsZero())
-			{
-				SourceLocation = SourceContext->SourceLocation;
-			}
-		}
-
-		if (!SourceLocation.IsZero())
-		{
-			if (FTargetingDefaultResultsSet* ResultData = FTargetingDefaultResultsSet::Find(TargetingHandle))
-			{
-#if ENABLE_DRAW_DEBUG
-				BuildPreSortDebugString(TargetingHandle, ResultData->TargetResults);
-#endif // ENABLE_DRAW_DEBUG
-
-				// compute the distances
-				const int32 NumTargets = ResultData->TargetResults.Num();
-				for (int32 TargetIterator = NumTargets - 1; TargetIterator >= 0; --TargetIterator)
-				{
-					FTargetingDefaultResultData& TargetResult = ResultData->TargetResults[TargetIterator];
-					if (AActor* TargetActor = TargetResult.HitResult.GetActor())
-					{
-						FVector TargetLocation = TargetActor->GetActorLocation();
-						TargetResult.Score = FVector::DistSquared(SourceLocation, TargetLocation);
-					}
-				}
-
-				// sort the set
-				ResultData->TargetResults.Sort([this](const FTargetingDefaultResultData& Lhs, const FTargetingDefaultResultData& Rhs)
-					{
-						if (bAscending)
-						{
-							return Lhs.Score < Rhs.Score;
-						}
-						else
-						{
-							return Lhs.Score > Rhs.Score;
-						}
-					});
-
-#if ENABLE_DRAW_DEBUG
-				BuildPostSortDebugString(TargetingHandle, ResultData->TargetResults);
-#endif // ENABLE_DRAW_DEBUG
-			}
-		}
+		FVector TargetLocation = TargetActor->GetActorLocation();
+		return FVector::DistSquared(SourceLocation, TargetLocation);
 	}
-
-	SetTaskAsyncState(TargetingHandle, ETargetingTaskAsyncState::Completed);
+	return 0.f;
 }
 
-#if ENABLE_DRAW_DEBUG
-
-void UTargetingFilterTask_SortByDistance::DrawDebug(UTargetingSubsystem* TargetingSubsystem, FTargetingDebugInfo& Info, const FTargetingRequestHandle& TargetingHandle, float XOffset, float YOffset, int32 MinTextRowsToAdvance) const
+FVector UTargetingFilterTask_SortByDistance::GetSourceLocation(const FTargetingRequestHandle& TargetingHandle) const
 {
-#if WITH_EDITORONLY_DATA
-	if (UTargetingSubsystem::IsTargetingDebugEnabled())
+	FVector SourceLocation = FVector::ZeroVector;
+	if (FTargetingSourceContext* SourceContext = FTargetingSourceContext::Find(TargetingHandle))
 	{
-		FTargetingDebugData& DebugData = FTargetingDebugData::FindOrAdd(TargetingHandle);
-		const FString& PreSortScratchPadString = DebugData.DebugScratchPadStrings.FindOrAdd(SortByDistanceConstants::PreSortPrefix + GetNameSafe(this));
-		const FString& PostSortScratchPadString = DebugData.DebugScratchPadStrings.FindOrAdd(SortByDistanceConstants::PostSortPrefix + GetNameSafe(this));
-		if (!PreSortScratchPadString.IsEmpty() && !PostSortScratchPadString.IsEmpty())
+		if (SourceContext->SourceActor)
 		{
-			if (Info.Canvas)
-			{
-				Info.Canvas->SetDrawColor(FColor::Yellow);
-			}
-
-			FString TaskString = FString::Printf(TEXT("Initial : %s"), *PreSortScratchPadString);
-			TargetingSubsystem->DebugLine(Info, TaskString, XOffset, YOffset, MinTextRowsToAdvance);
-
-			TaskString = FString::Printf(TEXT("Sorted : %s"), *PostSortScratchPadString);
-			TargetingSubsystem->DebugLine(Info, TaskString, XOffset, YOffset, MinTextRowsToAdvance);
-		}
-	}
-#endif // WITH_EDITORONLY_DATA
-}
-
-void UTargetingFilterTask_SortByDistance::BuildPreSortDebugString(const FTargetingRequestHandle& TargetingHandle, const TArray<FTargetingDefaultResultData>& TargetResults) const
-{
-#if WITH_EDITORONLY_DATA
-	if (UTargetingSubsystem::IsTargetingDebugEnabled())
-	{
-		FTargetingDebugData& DebugData = FTargetingDebugData::FindOrAdd(TargetingHandle);
-		FString& PreSortScratchPadString = DebugData.DebugScratchPadStrings.FindOrAdd(SortByDistanceConstants::PreSortPrefix + GetNameSafe(this));
-
-		for (const FTargetingDefaultResultData& TargetData : TargetResults)
-		{
-			if (const AActor* Target = TargetData.HitResult.GetActor())
-			{
-				if (PreSortScratchPadString.IsEmpty())
-				{
-					PreSortScratchPadString = FString::Printf(TEXT("%s"), *GetNameSafe(Target));
-				}
-				else
-				{
-					PreSortScratchPadString += FString::Printf(TEXT(", %s"), *GetNameSafe(Target));
-				}
-			}
+			SourceLocation = SourceContext->SourceActor->GetActorLocation();
 		}
-	}
-#endif // WITH_EDITORONLY_DATA
-}
-
-void UTargetingFilterTask_SortByDistance::BuildPostSortDebugString(const FTargetingRequestHandle& TargetingHandle, const TArray<FTargetingDefaultResultData>& TargetResults) const
-{
-#if WITH_EDITORONLY_DATA
-	if (UTargetingSubsystem::IsTargetingDebugEnabled())
-	{
-		FTargetingDebugData& DebugData = FTargetingDebugData::FindOrAdd(TargetingHandle);
-		FString& PostSortScratchPadString = DebugData.DebugScratchPadStrings.FindOrAdd(SortByDistanceConstants::PostSortPrefix + GetNameSafe(this));
-
-		for (const FTargetingDefaultResultData& TargetData : TargetResults)
+		else if (!SourceContext->SourceLocation.IsZero())
 		{
-			if (const AActor* Target = TargetData.HitResult.GetActor())
-			{
-				if (PostSortScratchPadString.IsEmpty())
-				{
-					PostSortScratchPadString = FString::Printf(TEXT("%s"), *GetNameSafe(Target));
-				}
-				else
-				{
-					PostSortScratchPadString += FString::Printf(TEXT(", %s"), *GetNameSafe(Target));
-				}
-			}
+			SourceLocation = SourceContext->SourceLocation;
 		}
 	}
-#endif // WITH_EDITORONLY_DATA
-}
-
-void UTargetingFilterTask_SortByDistance::ResetSortDebugStrings(const FTargetingRequestHandle& TargetingHandle) const
-{
-#if WITH_EDITORONLY_DATA
-	FTargetingDebugData& DebugData = FTargetingDebugData::FindOrAdd(TargetingHandle);
-	FString& PreSortScratchPadString = DebugData.DebugScratchPadStrings.FindOrAdd(SortByDistanceConstants::PreSortPrefix + GetNameSafe(this));
-	PreSortScratchPadString.Reset();
-
-	FString& PostSortScratchPadString = DebugData.DebugScratchPadStrings.FindOrAdd(SortByDistanceConstants::PostSortPrefix + GetNameSafe(this));
-	PostSortScratchPadString.Reset();
-#endif // WITH_EDITORONLY_DATA
-}
 
-#endif // ENABLE_DRAW_DEBUG
\ No newline at end of file
+	return SourceLocation;
+}
\ No newline at end of file
diff --git a/Engine/Plugins/Experimental/GameplayTargetingSystem/Source/GameplayTargetingSystem/Tasks/TargetingFilterTask_SortByDistance.h b/Engine/Plugins/Experimental/GameplayTargetingSystem/Source/GameplayTargetingSystem/Tasks/TargetingFilterTask_SortByDistance.h
--- a/Engine/Plugins/Experimental/GameplayTargetingSystem/Source/GameplayTargetingSystem/Tasks/TargetingFilterTask_SortByDistance.h
+++ b/Engine/Plugins/Experimental/GameplayTargetingSystem/Source/GameplayTargetingSystem/Tasks/TargetingFilterTask_SortByDistance.h
@@ -4,6 +4,7 @@
 #include "GameplayTargetingSystem/Types/TargetingSystemTypes.h"
 #include "TargetingTask.h"
 #include "UObject/Object.h"
+#include "TargetingSortTask_Base.h"
 
 #include "TargetingFilterTask_SortByDistance.generated.h"
 
@@ -22,28 +23,14 @@ struct FTargetingRequestHandle;
 *	and sort by that value.
 */
 UCLASS(Blueprintable)
-class UTargetingFilterTask_SortByDistance : public UTargetingTask
+class UTargetingFilterTask_SortByDistance : public UTargetingSortTask_Base
 {
 	GENERATED_BODY()
 
-public:
-	UTargetingFilterTask_SortByDistance(const FObjectInitializer& ObjectInitializer);
-
-	/** Evaluation function called by derived classes to process the targeting request */
-	virtual void Execute(const FTargetingRequestHandle& TargetingHandle) const override;
-
 protected:
-	UPROPERTY(EditAnywhere, Category = "Target Sorting | Data")
-	uint8 bAscending : 1;
+	virtual float GetScoreForTarget(const FTargetingRequestHandle& TargetingHandle, const FTargetingDefaultResultData& TargetData) const override;
 
-	/** Debug Helper Methods */
-#if ENABLE_DRAW_DEBUG
 private:
-	virtual void DrawDebug(UTargetingSubsystem* TargetingSubsystem, FTargetingDebugInfo& Info, const FTargetingRequestHandle& TargetingHandle, float XOffset, float YOffset, int32 MinTextRowsToAdvance) const override;
-	void BuildPreSortDebugString(const FTargetingRequestHandle& TargetingHandle, const TArray<FTargetingDefaultResultData>& TargetResults) const;
-	void BuildPostSortDebugString(const FTargetingRequestHandle& TargetingHandle, const TArray<FTargetingDefaultResultData>& TargetResults) const;
-	void ResetSortDebugStrings(const FTargetingRequestHandle& TargetingHandle) const;
-#endif // ENABLE_DRAW_DEBUG
-	/** ~Debug Helper Methods */
+	/** Returns the Source location, whether that comes from the source actor or the Source location in the Source Context */
+	FVector GetSourceLocation(const FTargetingRequestHandle& TargetingHandle) const;
 };
-
diff --git a/Engine/Plugins/Experimental/GameplayTargetingSystem/Source/GameplayTargetingSystem/Tasks/TargetingSortTask_Base.cpp b/Engine/Plugins/Experimental/GameplayTargetingSystem/Source/GameplayTargetingSystem/Tasks/TargetingSortTask_Base.cpp
new file mode 100644
--- /dev/null
+++ b/Engine/Plugins/Experimental/GameplayTargetingSystem/Source/GameplayTargetingSystem/Tasks/TargetingSortTask_Base.cpp
@@ -0,0 +1,181 @@
+// Copyright Epic Games, Inc. All Rights Reserved.
+
+#include "TargetingSortTask_Base.h"
+
+#include "GameplayTargetingSystem/TargetingSystem/TargetingSubsystem.h"
+#include "Kismet/KismetMathLibrary.h"
+
+#if ENABLE_DRAW_DEBUG
+#include "Engine/Canvas.h"
+#endif // ENABLE_DRAW_DEBUG
+
+namespace SortTaskConstants
+{
+	const FString PreSortPrefix = TEXT("PreSort");
+	const FString PostSortPrefix = TEXT("PostSort");
+}
+
+UTargetingSortTask_Base::UTargetingSortTask_Base(const FObjectInitializer& ObjectInitializer)
+	: Super(ObjectInitializer)
+{
+	bAscending = true;
+}
+
+float UTargetingSortTask_Base::GetScoreForTarget(const FTargetingRequestHandle& TargetingHandle, const FTargetingDefaultResultData& TargetData) const
+{
+	return 0.f;
+}
+
+void UTargetingSortTask_Base::Execute(const FTargetingRequestHandle& TargetingHandle) const
+{
+	Super::Execute(TargetingHandle);
+
+	SetTaskAsyncState(TargetingHandle, ETargetingTaskAsyncState::Executing);
+
+#if ENABLE_DRAW_DEBUG
+	ResetSortDebugStrings(TargetingHandle);
+#endif // ENABLE_DRAW_DEBUG
+
+	if (TargetingHandle.IsValid())
+	{
+		if (FTargetingDefaultResultsSet* ResultData = FTargetingDefaultResultsSet::Find(TargetingHandle))
+		{
+#if ENABLE_DRAW_DEBUG
+			BuildPreSortDebugString(TargetingHandle, ResultData->TargetResults);
+#endif // ENABLE_DRAW_DEBUG
+
+			const int32 NumTargets = ResultData->TargetResults.Num();
+
+			// We get the highest score first so we can normalize the score afterwards, 
+			// every task should have the same max score so none weights more than the others
+			float HighestScore = 0.f;
+			TArray<float> RawScores;
+			for (int32 TargetIterator = NumTargets - 1; TargetIterator >= 0; --TargetIterator)
+			{
+				FTargetingDefaultResultData& TargetResult = ResultData->TargetResults[TargetIterator];
+				float RawScore = GetScoreForTarget(TargetingHandle, TargetResult);
+				RawScores.Add(RawScore);
+				HighestScore = FMath::Max(HighestScore, RawScore);
+			}
+
+			if(ensureMsgf(NumTargets == RawScores.Num(), TEXT("The cached raw scores should be the same size as the number of targets!")))
+			{
+				// Adding the normalized scores to each target result.
+				for (int32 TargetIterator = NumTargets - 1; TargetIterator >= 0; --TargetIterator)
+				{
+					FTargetingDefaultResultData& TargetResult = ResultData->TargetResults[TargetIterator];
+					
+					// Driving ascending/descending sorting based on a multiplier so it carries over to other tasks 
+					float SortingMultiplier = bAscending ? 1.f : -1.f;
+					TargetResult.Score += UKismetMathLibrary::SafeDivide(RawScores[TargetIterator], HighestScore) * SortingMultiplier;
+				}
+			}
+
+			// sort the set
+			ResultData->TargetResults.Sort([this](const FTargetingDefaultResultData& Lhs, const FTargetingDefaultResultData& Rhs)
+				{
+					return Lhs.Score < Rhs.Score;
+				});
+
+#if ENABLE_DRAW_DEBUG
+			BuildPostSortDebugString(TargetingHandle, ResultData->TargetResults);
+#endif // ENABLE_DRAW_DEBUG
+		}
+	}
+
+	SetTaskAsyncState(TargetingHandle, ETargetingTaskAsyncState::Completed);
+}
+
+
+#if ENABLE_DRAW_DEBUG
+
+void UTargetingSortTask_Base::DrawDebug(UTargetingSubsystem* TargetingSubsystem, FTargetingDebugInfo& Info, const FTargetingRequestHandle& TargetingHandle, float XOffset, float YOffset, int32 MinTextRowsToAdvance) const
+{
+#if WITH_EDITORONLY_DATA
+	if (TargetingSubsystem && UTargetingSubsystem::IsTargetingDebugEnabled())
+	{
+		FTargetingDebugData& DebugData = FTargetingDebugData::FindOrAdd(TargetingHandle);
+		const FString& PreSortScratchPadString = DebugData.DebugScratchPadStrings.FindOrAdd(SortTaskConstants::PreSortPrefix + GetNameSafe(this));
+		const FString& PostSortScratchPadString = DebugData.DebugScratchPadStrings.FindOrAdd(SortTaskConstants::PostSortPrefix + GetNameSafe(this));
+		if (!PreSortScratchPadString.IsEmpty() && !PostSortScratchPadString.IsEmpty())
+		{
+			if (Info.Canvas)
+			{
+				Info.Canvas->SetDrawColor(FColor::Yellow);
+			}
+
+			FString TaskString = FString::Printf(TEXT("Initial : %s"), *PreSortScratchPadString);
+			TargetingSubsystem->DebugLine(Info, TaskString, XOffset, YOffset, MinTextRowsToAdvance);
+
+			TaskString = FString::Printf(TEXT("Sorted : %s"), *PostSortScratchPadString);
+			TargetingSubsystem->DebugLine(Info, TaskString, XOffset, YOffset, MinTextRowsToAdvance);
+		}
+	}
+#endif // WITH_EDITORONLY_DATA
+}
+
+void UTargetingSortTask_Base::BuildPreSortDebugString(const FTargetingRequestHandle& TargetingHandle, const TArray<FTargetingDefaultResultData>& TargetResults) const
+{
+#if WITH_EDITORONLY_DATA
+	if (UTargetingSubsystem::IsTargetingDebugEnabled())
+	{
+		FTargetingDebugData& DebugData = FTargetingDebugData::FindOrAdd(TargetingHandle);
+		FString& PreSortScratchPadString = DebugData.DebugScratchPadStrings.FindOrAdd(SortTaskConstants::PreSortPrefix + GetNameSafe(this));
+
+		for (const FTargetingDefaultResultData& TargetData : TargetResults)
+		{
+			if (const AActor* Target = TargetData.HitResult.GetActor())
+			{
+				if (PreSortScratchPadString.IsEmpty())
+				{
+					PreSortScratchPadString = *GetNameSafe(Target);
+				}
+				else
+				{
+					PreSortScratchPadString += FString::Printf(TEXT(", %s"), *GetNameSafe(Target));
+				}
+			}
+		}
+	}
+#endif // WITH_EDITORONLY_DATA
+}
+
+void UTargetingSortTask_Base::BuildPostSortDebugString(const FTargetingRequestHandle& TargetingHandle, const TArray<FTargetingDefaultResultData>& TargetResults) const
+{
+#if WITH_EDITORONLY_DATA
+	if (UTargetingSubsystem::IsTargetingDebugEnabled())
+	{
+		FTargetingDebugData& DebugData = FTargetingDebugData::FindOrAdd(TargetingHandle);
+		FString& PostSortScratchPadString = DebugData.DebugScratchPadStrings.FindOrAdd(SortTaskConstants::PostSortPrefix + GetNameSafe(this));
+
+		for (const FTargetingDefaultResultData& TargetData : TargetResults)
+		{
+			if (const AActor* Target = TargetData.HitResult.GetActor())
+			{
+				if (PostSortScratchPadString.IsEmpty())
+				{
+					PostSortScratchPadString = *GetNameSafe(Target);
+				}
+				else
+				{
+					PostSortScratchPadString += FString::Printf(TEXT(", %s"), *GetNameSafe(Target));
+				}
+			}
+		}
+	}
+#endif // WITH_EDITORONLY_DATA
+}
+
+void UTargetingSortTask_Base::ResetSortDebugStrings(const FTargetingRequestHandle& TargetingHandle) const
+{
+#if WITH_EDITORONLY_DATA
+	FTargetingDebugData& DebugData = FTargetingDebugData::FindOrAdd(TargetingHandle);
+	FString& PreSortScratchPadString = DebugData.DebugScratchPadStrings.FindOrAdd(SortTaskConstants::PreSortPrefix + GetNameSafe(this));
+	PreSortScratchPadString.Reset();
+
+	FString& PostSortScratchPadString = DebugData.DebugScratchPadStrings.FindOrAdd(SortTaskConstants::PostSortPrefix + GetNameSafe(this));
+	PostSortScratchPadString.Reset();
+#endif // WITH_EDITORONLY_DATA
+}
+
+#endif // ENABLE_DRAW_DEBUG
diff --git a/Engine/Plugins/Experimental/GameplayTargetingSystem/Source/GameplayTargetingSystem/Tasks/TargetingSortTask_Base.h b/Engine/Plugins/Experimental/GameplayTargetingSystem/Source/GameplayTargetingSystem/Tasks/TargetingSortTask_Base.h
new file mode 100644
--- /dev/null
+++ b/Engine/Plugins/Experimental/GameplayTargetingSystem/Source/GameplayTargetingSystem/Tasks/TargetingSortTask_Base.h
@@ -0,0 +1,43 @@
+// Copyright Epic Games, Inc. All Rights Reserved.
+
+#pragma once
+
+#include "TargetingTask.h"
+
+#include "TargetingSortTask_Base.generated.h"
+
+/** 
+*	@class TargetingSortTask_Base
+*
+*	A base class that has the basic setup for a sort task.
+*/
+UCLASS(Abstract)
+class TARGETINGSYSTEM_API UTargetingSortTask_Base : public UTargetingTask
+{
+    GENERATED_BODY()
+
+public:
+	UTargetingSortTask_Base(const FObjectInitializer& ObjectInitializer);
+
+protected:
+	/** Called on every target to get a Score for sorting. This score will be added to the Score float in FTargetingDefaultResultData */
+	virtual float GetScoreForTarget(const FTargetingRequestHandle& TargetingHandle, const FTargetingDefaultResultData& TargetData) const;
+
+private:
+	/** Evaluation function called by derived classes to process the targeting request */
+	virtual void Execute(const FTargetingRequestHandle& TargetingHandle) const override;
+
+protected:
+	UPROPERTY(EditAnywhere, Category = "Target Sorting | Data")
+	uint8 bAscending : 1;
+
+	/** Debug Helper Methods */
+#if ENABLE_DRAW_DEBUG
+private:
+	virtual void DrawDebug(UTargetingSubsystem* TargetingSubsystem, FTargetingDebugInfo& Info, const FTargetingRequestHandle& TargetingHandle, float XOffset, float YOffset, int32 MinTextRowsToAdvance) const override;
+	void BuildPreSortDebugString(const FTargetingRequestHandle& TargetingHandle, const TArray<FTargetingDefaultResultData>& TargetResults) const;
+	void BuildPostSortDebugString(const FTargetingRequestHandle& TargetingHandle, const TArray<FTargetingDefaultResultData>& TargetResults) const;
+	void ResetSortDebugStrings(const FTargetingRequestHandle& TargetingHandle) const;
+#endif // ENABLE_DRAW_DEBUG
+	/** ~Debug Helper Methods */
+};
