Fix non-unity
#rnx

diff --git a/Engine/Plugins/Experimental/GameplayTargetingSystem/Source/GameplayTargetingSystem/Tasks/TargetingSortTask_Base.cpp b/Engine/Plugins/Experimental/GameplayTargetingSystem/Source/GameplayTargetingSystem/Tasks/TargetingSortTask_Base.cpp
--- a/Engine/Plugins/Experimental/GameplayTargetingSystem/Source/GameplayTargetingSystem/Tasks/TargetingSortTask_Base.cpp
+++ b/Engine/Plugins/Experimental/GameplayTargetingSystem/Source/GameplayTargetingSystem/Tasks/TargetingSortTask_Base.cpp
@@ -2,6 +2,7 @@
 
 #include "TargetingSortTask_Base.h"
 
+#include "GameFramework/Actor.h"
 #include "GameplayTargetingSystem/TargetingSystem/TargetingSubsystem.h"
 #include "Kismet/KismetMathLibrary.h"
 
