Add null checks to Targeting async actions to avoid hang when performing them

diff --git a/Engine/Plugins/Experimental/GameplayTargetingSystem/Source/GameplayTargetingSystem/AbilityTasks/AbilityTask_PerformTargeting.cpp b/Engine/Plugins/Experimental/GameplayTargetingSystem/Source/GameplayTargetingSystem/AbilityTasks/AbilityTask_PerformTargeting.cpp
--- a/Engine/Plugins/Experimental/GameplayTargetingSystem/Source/GameplayTargetingSystem/AbilityTasks/AbilityTask_PerformTargeting.cpp
+++ b/Engine/Plugins/Experimental/GameplayTargetingSystem/Source/GameplayTargetingSystem/AbilityTasks/AbilityTask_PerformTargeting.cpp
@@ -52,6 +52,11 @@ void UAbilityTask_PerformTargeting::Activate()
 
 UAbilityTask_PerformTargeting* UAbilityTask_PerformTargeting::PerformTargetingRequest(UGameplayAbility* OwningAbility, UTargetingPreset* InTargetingPreset, bool bAllowAsync)
 {
+	if (!InTargetingPreset)
+	{
+		return nullptr;
+	}
+
 	UAbilityTask_PerformTargeting* Task = NewAbilityTask<UAbilityTask_PerformTargeting>(OwningAbility);
 	Task->TargetingPreset = InTargetingPreset;
 	Task->bPerformAsync = bAllowAsync;
@@ -61,6 +66,11 @@ UAbilityTask_PerformTargeting* UAbilityTask_PerformTargeting::PerformTargetingRe
 
 UAbilityTask_PerformTargeting* UAbilityTask_PerformTargeting::PerformFilteringRequest(UGameplayAbility* OwningAbility, UTargetingPreset* InTargetingPreset, const TArray<AActor*> InTargets, bool bAllowAsync)
 {
+	if (!InTargetingPreset)
+	{
+		return nullptr;
+	}
+
 	UAbilityTask_PerformTargeting* Task = NewAbilityTask<UAbilityTask_PerformTargeting>(OwningAbility);
 	Task->TargetingPreset = InTargetingPreset;
 	Task->InitialTargets = InTargets;
diff --git a/Engine/Plugins/Experimental/GameplayTargetingSystem/Source/GameplayTargetingSystem/Async/AsyncAction_PerformTargeting.cpp b/Engine/Plugins/Experimental/GameplayTargetingSystem/Source/GameplayTargetingSystem/Async/AsyncAction_PerformTargeting.cpp
--- a/Engine/Plugins/Experimental/GameplayTargetingSystem/Source/GameplayTargetingSystem/Async/AsyncAction_PerformTargeting.cpp
+++ b/Engine/Plugins/Experimental/GameplayTargetingSystem/Source/GameplayTargetingSystem/Async/AsyncAction_PerformTargeting.cpp
@@ -19,6 +19,11 @@ UAsyncAction_PerformTargeting::UAsyncAction_PerformTargeting(const FObjectInitia
 
 UAsyncAction_PerformTargeting* UAsyncAction_PerformTargeting::PerformTargetingRequest(AActor* SourceActor, UTargetingPreset* TargetingPreset, bool bUseAsyncTargeting)
 {
+	if (!TargetingPreset)
+	{
+		return nullptr;
+	}
+
 	UWorld* World = GEngine->GetWorldFromContextObject(SourceActor, EGetWorldErrorMode::LogAndReturnNull);
 	if (!World)
 	{
@@ -36,6 +41,11 @@ UAsyncAction_PerformTargeting* UAsyncAction_PerformTargeting::PerformTargetingRe
 
 UAsyncAction_PerformTargeting* UAsyncAction_PerformTargeting::PerformFilteringRequest(AActor* SourceActor, UTargetingPreset* TargetingPreset, bool bUseAsyncTargeting, const TArray<AActor*> InTargets)
 {
+	if (!TargetingPreset)
+	{
+		return nullptr;
+	}
+
 	UWorld* World = GEngine->GetWorldFromContextObject(SourceActor, EGetWorldErrorMode::LogAndReturnNull);
 	if (!World)
 	{
