Fix some TargetingSubsystem functions causing un-fixable compiler errors

diff --git a/Engine/Plugins/Experimental/GameplayTargetingSystem/Source/GameplayTargetingSystem/TargetingSystem/TargetingSubsystem.cpp b/Engine/Plugins/Experimental/GameplayTargetingSystem/Source/GameplayTargetingSystem/TargetingSystem/TargetingSubsystem.cpp
--- a/Engine/Plugins/Experimental/GameplayTargetingSystem/Source/GameplayTargetingSystem/TargetingSystem/TargetingSubsystem.cpp
+++ b/Engine/Plugins/Experimental/GameplayTargetingSystem/Source/GameplayTargetingSystem/TargetingSystem/TargetingSubsystem.cpp
@@ -609,7 +609,6 @@ FTargetingSourceContext UTargetingSubsystem::GetTargetingSourceContext(FTargetin
 	return FTargetingSourceContext();
 }
 
-/* static */
 void UTargetingSubsystem::GetTargetingResultsActors(FTargetingRequestHandle TargetingHandle, TArray<AActor*>& Targets)
 {
 	if (TargetingHandle.IsValid())
@@ -627,7 +626,6 @@ void UTargetingSubsystem::GetTargetingResultsActors(FTargetingRequestHandle Targ
 	}
 }
 
-/* static */
 void UTargetingSubsystem::GetTargetingResults(FTargetingRequestHandle TargetingHandle, TArray<FHitResult>& OutTargets)
 {
 	if (TargetingHandle.IsValid())
diff --git a/Engine/Plugins/Experimental/GameplayTargetingSystem/Source/GameplayTargetingSystem/TargetingSystem/TargetingSubsystem.h b/Engine/Plugins/Experimental/GameplayTargetingSystem/Source/GameplayTargetingSystem/TargetingSystem/TargetingSubsystem.h
--- a/Engine/Plugins/Experimental/GameplayTargetingSystem/Source/GameplayTargetingSystem/TargetingSystem/TargetingSubsystem.h
+++ b/Engine/Plugins/Experimental/GameplayTargetingSystem/Source/GameplayTargetingSystem/TargetingSystem/TargetingSubsystem.h
@@ -173,15 +173,15 @@ private:
 public:
 	/** Returns the targeting source context for the targeting request handle */
 	UFUNCTION(BlueprintCallable, Category = "Targeting System | Targeting Types")
-	static FTargetingSourceContext GetTargetingSourceContext(FTargetingRequestHandle TargetingHandle);
+	FTargetingSourceContext GetTargetingSourceContext(FTargetingRequestHandle TargetingHandle);
 
 	/** Method to get the actor targets from a given targeting request handle */
 	UFUNCTION(BlueprintCallable, Category = "Targeting System | Targeting Results")
-	static TARGETINGSYSTEM_API void GetTargetingResultsActors(FTargetingRequestHandle TargetingHandle, TArray<AActor*>& Targets);
+	TARGETINGSYSTEM_API void GetTargetingResultsActors(FTargetingRequestHandle TargetingHandle, TArray<AActor*>& Targets);
 
 	/** Helper method to get the set of hit results for a given targeting handle */
 	UFUNCTION(BlueprintCallable, Category = "Targeting System | Targeting Results")
-	static TARGETINGSYSTEM_API void GetTargetingResults(FTargetingRequestHandle TargetingHandle, TArray<FHitResult>& OutTargets);
+	TARGETINGSYSTEM_API void GetTargetingResults(FTargetingRequestHandle TargetingHandle, TArray<FHitResult>& OutTargets);
 
 	/** ~Blueprint Helper Methods */
 
