Fix cook error caused by missing const specifier

diff --git a/Engine/Plugins/Experimental/GameplayTargetingSystem/Source/GameplayTargetingSystem/TargetingSystem/TargetingSubsystem.cpp b/Engine/Plugins/Experimental/GameplayTargetingSystem/Source/GameplayTargetingSystem/TargetingSystem/TargetingSubsystem.cpp
--- a/Engine/Plugins/Experimental/GameplayTargetingSystem/Source/GameplayTargetingSystem/TargetingSystem/TargetingSubsystem.cpp
+++ b/Engine/Plugins/Experimental/GameplayTargetingSystem/Source/GameplayTargetingSystem/TargetingSystem/TargetingSubsystem.cpp
@@ -598,7 +598,7 @@ void UTargetingSubsystem::ProcessTargetingRequestTasks(FTargetingRequestHandle&
 	}
 }
 
-FTargetingSourceContext UTargetingSubsystem::GetTargetingSourceContext(FTargetingRequestHandle TargetingHandle)
+FTargetingSourceContext UTargetingSubsystem::GetTargetingSourceContext(FTargetingRequestHandle TargetingHandle) const
 {
 	if (TargetingHandle.IsValid())
 	{
@@ -609,7 +609,7 @@ FTargetingSourceContext UTargetingSubsystem::GetTargetingSourceContext(FTargetin
 	return FTargetingSourceContext();
 }
 
-void UTargetingSubsystem::GetTargetingResultsActors(FTargetingRequestHandle TargetingHandle, TArray<AActor*>& Targets)
+void UTargetingSubsystem::GetTargetingResultsActors(FTargetingRequestHandle TargetingHandle, TArray<AActor*>& Targets) const
 {
 	if (TargetingHandle.IsValid())
 	{
@@ -626,7 +626,7 @@ void UTargetingSubsystem::GetTargetingResultsActors(FTargetingRequestHandle Targ
 	}
 }
 
-void UTargetingSubsystem::GetTargetingResults(FTargetingRequestHandle TargetingHandle, TArray<FHitResult>& OutTargets)
+void UTargetingSubsystem::GetTargetingResults(FTargetingRequestHandle TargetingHandle, TArray<FHitResult>& OutTargets) const
 {
 	if (TargetingHandle.IsValid())
 	{
diff --git a/Engine/Plugins/Experimental/GameplayTargetingSystem/Source/GameplayTargetingSystem/TargetingSystem/TargetingSubsystem.h b/Engine/Plugins/Experimental/GameplayTargetingSystem/Source/GameplayTargetingSystem/TargetingSystem/TargetingSubsystem.h
--- a/Engine/Plugins/Experimental/GameplayTargetingSystem/Source/GameplayTargetingSystem/TargetingSystem/TargetingSubsystem.h
+++ b/Engine/Plugins/Experimental/GameplayTargetingSystem/Source/GameplayTargetingSystem/TargetingSystem/TargetingSubsystem.h
@@ -173,15 +173,15 @@ private:
 public:
 	/** Returns the targeting source context for the targeting request handle */
 	UFUNCTION(BlueprintCallable, Category = "Targeting System | Targeting Types")
-	FTargetingSourceContext GetTargetingSourceContext(FTargetingRequestHandle TargetingHandle);
+	FTargetingSourceContext GetTargetingSourceContext(FTargetingRequestHandle TargetingHandle) const;
 
 	/** Method to get the actor targets from a given targeting request handle */
 	UFUNCTION(BlueprintCallable, Category = "Targeting System | Targeting Results")
-	TARGETINGSYSTEM_API void GetTargetingResultsActors(FTargetingRequestHandle TargetingHandle, TArray<AActor*>& Targets);
+	TARGETINGSYSTEM_API void GetTargetingResultsActors(FTargetingRequestHandle TargetingHandle, TArray<AActor*>& Targets) const;
 
 	/** Helper method to get the set of hit results for a given targeting handle */
 	UFUNCTION(BlueprintCallable, Category = "Targeting System | Targeting Results")
-	TARGETINGSYSTEM_API void GetTargetingResults(FTargetingRequestHandle TargetingHandle, TArray<FHitResult>& OutTargets);
+	TARGETINGSYSTEM_API void GetTargetingResults(FTargetingRequestHandle TargetingHandle, TArray<FHitResult>& OutTargets) const;
 
 	/** ~Blueprint Helper Methods */
 
