Cleanup beta tags, versions, and descriptions of gameplay plugins to better match the actual state

diff --git a/Engine/Plugins/Experimental/AbilitySystemGameFeatureActions/AbilitySystemGameFeatureActions.uplugin b/Engine/Plugins/Experimental/AbilitySystemGameFeatureActions/AbilitySystemGameFeatureActions.uplugin
--- a/Engine/Plugins/Experimental/AbilitySystemGameFeatureActions/AbilitySystemGameFeatureActions.uplugin
+++ b/Engine/Plugins/Experimental/AbilitySystemGameFeatureActions/AbilitySystemGameFeatureActions.uplugin
@@ -2,8 +2,8 @@
 	"FileVersion": 3,
 	"Version": 1,
 	"VersionName": "0.1",
-	"FriendlyName": "Ability System Game Feature Actions",
-	"Description": "Ability System Game Feature Actions",
+	"FriendlyName": "Gameplay Abilities Game Feature Actions",
+	"Description": "Game feature actions to support modular use of the gameplay abilities system",
 	"Category": "Gameplay",
 	"CreatedBy": "Epic Games, Inc.",
 	"CreatedByURL": "https://epicgames.com",
@@ -12,7 +12,7 @@
 	"SupportURL": "",
 	"EnabledByDefault": false,
 	"CanContainContent": false,
-	"IsBetaVersion": false,
+	"IsBetaVersion": true,
 	"Installed": false,
 	"Modules": [
 		{
diff --git a/Engine/Plugins/Experimental/GameFeatures/GameFeatures.uplugin b/Engine/Plugins/Experimental/GameFeatures/GameFeatures.uplugin
--- a/Engine/Plugins/Experimental/GameFeatures/GameFeatures.uplugin
+++ b/Engine/Plugins/Experimental/GameFeatures/GameFeatures.uplugin
@@ -1,7 +1,7 @@
 {
 	"FileVersion": 3,
 	"Version": 1,
-	"VersionName": "0.1",
+	"VersionName": "1.0",
 	"FriendlyName": "Game Features",
 	"Description": "Support for modular Game Feature Plugins",
 	"Category": "Gameplay",
@@ -12,7 +12,7 @@
 	"SupportURL": "",
 	"EnabledByDefault": false,
 	"CanContainContent": false,
-	"IsBetaVersion": false,
+	"IsBetaVersion": true,
 	"Installed": false,
 	"Modules": [
 		{
diff --git a/Engine/Plugins/Experimental/GameplayTargetingSystem/TargetingSystem.uplugin b/Engine/Plugins/Experimental/GameplayTargetingSystem/TargetingSystem.uplugin
--- a/Engine/Plugins/Experimental/GameplayTargetingSystem/TargetingSystem.uplugin
+++ b/Engine/Plugins/Experimental/GameplayTargetingSystem/TargetingSystem.uplugin
@@ -2,8 +2,8 @@
 	"FileVersion": 3,
 	"Version": 1,
 	"VersionName": "1.0",
-	"FriendlyName": "TargetingSystem",
-	"Description": "Generic Targeting System to be used for gameplay. Abilities/Aim Assist/etc",
+	"FriendlyName": "Targeting System",
+	"Description": "Generic targeting system for use with gameplay abilities, aim assist, etc",
 	"Category": "Gameplay",
 	"CreatedBy": "Epic Games, Inc.",
 	"CreatedByURL": "https://epicgames.com",
@@ -12,7 +12,7 @@
 	"SupportURL": "",
 	"EnabledByDefault": false,
 	"CanContainContent": false,
-	"IsBetaVersion": false,
+	"IsBetaVersion": true,
 	"Installed": false,
 	"Modules": [
 		{
@@ -26,6 +26,5 @@
 			"Name": "GameplayAbilities",
 			"Enabled": true
 		}
-	],
-	"IsExperimentalVersion": false
+	]
 }
\ No newline at end of file
diff --git a/Engine/Plugins/Experimental/ModularGameplay/ModularGameplay.uplugin b/Engine/Plugins/Experimental/ModularGameplay/ModularGameplay.uplugin
--- a/Engine/Plugins/Experimental/ModularGameplay/ModularGameplay.uplugin
+++ b/Engine/Plugins/Experimental/ModularGameplay/ModularGameplay.uplugin
@@ -3,7 +3,7 @@
 	"Version": 1,
 	"VersionName": "1.0",
 	"FriendlyName": "Modular Gameplay",
-	"Description": "",
+	"Description": "Base classes and subsystems to support modular use of the gameplay framework",
 	"Category": "Gameplay",
 	"CreatedBy": "Epic Games, Inc.",
 	"CreatedByURL": "https://epicgames.com",
@@ -12,7 +12,7 @@
 	"SupportURL": "",
 	"EnabledByDefault": false,
 	"CanContainContent": false,
-	"IsBetaVersion": false,
+	"IsBetaVersion": true,
 	"Installed": false,
 	"Modules": [
 		{
diff --git a/Engine/Plugins/Runtime/GameplayAbilities/GameplayAbilities.uplugin b/Engine/Plugins/Runtime/GameplayAbilities/GameplayAbilities.uplugin
--- a/Engine/Plugins/Runtime/GameplayAbilities/GameplayAbilities.uplugin
+++ b/Engine/Plugins/Runtime/GameplayAbilities/GameplayAbilities.uplugin
@@ -1,7 +1,7 @@
 {
 	"FileVersion" : 3,
 	"Version" : 1,
-	"VersionName" : "0.1",
+	"VersionName" : "1.0",
 	"FriendlyName" : "Gameplay Abilities",
 	"Description" : "Adds GameplayEffect and GameplayAbility classes to handle complicated gameplay interactions.",
 	"Category" : "Gameplay",
