GameplayTargetingSystem: Fix relative include paths from module parent directory

#rnx

diff --git a/Engine/Plugins/Experimental/GameplayTargetingSystem/Source/GameplayTargetingSystem/AbilityTasks/AbilityTask_PerformTargeting.cpp b/Engine/Plugins/Experimental/GameplayTargetingSystem/Source/GameplayTargetingSystem/AbilityTasks/AbilityTask_PerformTargeting.cpp
--- a/Engine/Plugins/Experimental/GameplayTargetingSystem/Source/GameplayTargetingSystem/AbilityTasks/AbilityTask_PerformTargeting.cpp
+++ b/Engine/Plugins/Experimental/GameplayTargetingSystem/Source/GameplayTargetingSystem/AbilityTasks/AbilityTask_PerformTargeting.cpp
@@ -1,7 +1,7 @@
 // Copyright Epic Games, Inc. All Rights Reserved.
 #include "AbilityTask_PerformTargeting.h"
-#include "GameplayTargetingSystem/TargetingSystem/TargetingSubsystem.h"
-#include "GameplayTargetingSystem/Types/TargetingSystemLogs.h"
+#include "TargetingSystem/TargetingSubsystem.h"
+#include "Types/TargetingSystemLogs.h"
 
 void UAbilityTask_PerformTargeting::Activate()
 {
diff --git a/Engine/Plugins/Experimental/GameplayTargetingSystem/Source/GameplayTargetingSystem/AbilityTasks/AbilityTask_PerformTargeting.h b/Engine/Plugins/Experimental/GameplayTargetingSystem/Source/GameplayTargetingSystem/AbilityTasks/AbilityTask_PerformTargeting.h
--- a/Engine/Plugins/Experimental/GameplayTargetingSystem/Source/GameplayTargetingSystem/AbilityTasks/AbilityTask_PerformTargeting.h
+++ b/Engine/Plugins/Experimental/GameplayTargetingSystem/Source/GameplayTargetingSystem/AbilityTasks/AbilityTask_PerformTargeting.h
@@ -5,7 +5,7 @@
 #include "UObject/ObjectMacros.h"
 #include "Engine/EngineTypes.h"
 #include "Abilities/Tasks/AbilityTask.h"
-#include "GameplayTargetingSystem/Types/TargetingSystemTypes.h"
+#include "Types/TargetingSystemTypes.h"
 #include "AbilityTask_PerformTargeting.generated.h"
 
 class UTargetingPreset;
diff --git a/Engine/Plugins/Experimental/GameplayTargetingSystem/Source/GameplayTargetingSystem/Async/AsyncAction_PerformTargeting.cpp b/Engine/Plugins/Experimental/GameplayTargetingSystem/Source/GameplayTargetingSystem/Async/AsyncAction_PerformTargeting.cpp
--- a/Engine/Plugins/Experimental/GameplayTargetingSystem/Source/GameplayTargetingSystem/Async/AsyncAction_PerformTargeting.cpp
+++ b/Engine/Plugins/Experimental/GameplayTargetingSystem/Source/GameplayTargetingSystem/Async/AsyncAction_PerformTargeting.cpp
@@ -5,10 +5,10 @@
 #include "Engine/Engine.h"
 #include "Engine/World.h"
 #include "GameFramework/Actor.h"
-#include "GameplayTargetingSystem/TargetingSystem/TargetingPreset.h"
-#include "GameplayTargetingSystem/TargetingSystem/TargetingSubsystem.h"
-#include "GameplayTargetingSystem/Types/TargetingSystemLogs.h"
-#include "GameplayTargetingSystem/Types/TargetingSystemTypes.h"
+#include "TargetingSystem/TargetingPreset.h"
+#include "TargetingSystem/TargetingSubsystem.h"
+#include "Types/TargetingSystemLogs.h"
+#include "Types/TargetingSystemTypes.h"
 
 
 UAsyncAction_PerformTargeting::UAsyncAction_PerformTargeting(const FObjectInitializer& ObjectInitializer)
diff --git a/Engine/Plugins/Experimental/GameplayTargetingSystem/Source/GameplayTargetingSystem/Async/AsyncAction_PerformTargeting.h b/Engine/Plugins/Experimental/GameplayTargetingSystem/Source/GameplayTargetingSystem/Async/AsyncAction_PerformTargeting.h
--- a/Engine/Plugins/Experimental/GameplayTargetingSystem/Source/GameplayTargetingSystem/Async/AsyncAction_PerformTargeting.h
+++ b/Engine/Plugins/Experimental/GameplayTargetingSystem/Source/GameplayTargetingSystem/Async/AsyncAction_PerformTargeting.h
@@ -3,7 +3,7 @@
 #pragma once
 
 #include "CoreMinimal.h"
-#include "GameplayTargetingSystem/Types/TargetingSystemTypes.h"
+#include "Types/TargetingSystemTypes.h"
 #include "Kismet/BlueprintAsyncActionBase.h"
 #include "UObject/ObjectMacros.h"
 
diff --git a/Engine/Plugins/Experimental/GameplayTargetingSystem/Source/GameplayTargetingSystem/TargetingSystem.Build.cs b/Engine/Plugins/Experimental/GameplayTargetingSystem/Source/GameplayTargetingSystem/TargetingSystem.Build.cs
--- a/Engine/Plugins/Experimental/GameplayTargetingSystem/Source/GameplayTargetingSystem/TargetingSystem.Build.cs
+++ b/Engine/Plugins/Experimental/GameplayTargetingSystem/Source/GameplayTargetingSystem/TargetingSystem.Build.cs
@@ -2,18 +2,21 @@
 using UnrealBuildTool;
 public class TargetingSystem : ModuleRules
 {
+	// TODO: Rename parent directory to TargetingSystem so it matches module name
 	public TargetingSystem(ReadOnlyTargetRules Target) : base(Target)
 	{
 		PCHUsage = ModuleRules.PCHUsageMode.UseExplicitOrSharedPCHs;
 
+		// TODO: Move source code to have proper Private/Public subfolders
 		PublicIncludePaths.AddRange(
 			new string[] {
+				ModuleDirectory
 			}
 		);
 
 		PrivateIncludePaths.AddRange(
 			new string[] {
-            }
+			}
 		);
 
 		PublicDependencyModuleNames.AddRange(
diff --git a/Engine/Plugins/Experimental/GameplayTargetingSystem/Source/GameplayTargetingSystem/TargetingSystem/TargetingPreset.h b/Engine/Plugins/Experimental/GameplayTargetingSystem/Source/GameplayTargetingSystem/TargetingSystem/TargetingPreset.h
--- a/Engine/Plugins/Experimental/GameplayTargetingSystem/Source/GameplayTargetingSystem/TargetingSystem/TargetingPreset.h
+++ b/Engine/Plugins/Experimental/GameplayTargetingSystem/Source/GameplayTargetingSystem/TargetingSystem/TargetingPreset.h
@@ -3,7 +3,7 @@
 
 #include "UObject/Object.h"
 #include "Engine/DataAsset.h"
-#include "GameplayTargetingSystem/Types/TargetingSystemTypes.h"
+#include "Types/TargetingSystemTypes.h"
 
 #include "TargetingPreset.generated.h"
 
diff --git a/Engine/Plugins/Experimental/GameplayTargetingSystem/Source/GameplayTargetingSystem/TargetingSystem/TargetingSubsystem.cpp b/Engine/Plugins/Experimental/GameplayTargetingSystem/Source/GameplayTargetingSystem/TargetingSystem/TargetingSubsystem.cpp
--- a/Engine/Plugins/Experimental/GameplayTargetingSystem/Source/GameplayTargetingSystem/TargetingSystem/TargetingSubsystem.cpp
+++ b/Engine/Plugins/Experimental/GameplayTargetingSystem/Source/GameplayTargetingSystem/TargetingSystem/TargetingSubsystem.cpp
@@ -4,8 +4,8 @@
 #include "Engine/Engine.h"
 #include "Engine/GameInstance.h"
 #include "Engine/World.h"
-#include "GameplayTargetingSystem/Types/TargetingSystemLogs.h"
-#include "GameplayTargetingSystem/Types/TargetingSystemTypes.h"
+#include "Types/TargetingSystemLogs.h"
+#include "Types/TargetingSystemTypes.h"
 #include "HAL/IConsoleManager.h"
 #include "Kismet/KismetStringLibrary.h"
 #include "ProfilingDebugging/CsvProfiler.h"
diff --git a/Engine/Plugins/Experimental/GameplayTargetingSystem/Source/GameplayTargetingSystem/TargetingSystem/TargetingSubsystem.h b/Engine/Plugins/Experimental/GameplayTargetingSystem/Source/GameplayTargetingSystem/TargetingSystem/TargetingSubsystem.h
--- a/Engine/Plugins/Experimental/GameplayTargetingSystem/Source/GameplayTargetingSystem/TargetingSystem/TargetingSubsystem.h
+++ b/Engine/Plugins/Experimental/GameplayTargetingSystem/Source/GameplayTargetingSystem/TargetingSystem/TargetingSubsystem.h
@@ -3,8 +3,8 @@
 
 #include "CoreMinimal.h"
 #include "DrawDebugHelpers.h"
-#include "GameplayTargetingSystem/Tasks/TargetingTask.h"
-#include "GameplayTargetingSystem/Types/TargetingSystemTypes.h"
+#include "Tasks/TargetingTask.h"
+#include "Types/TargetingSystemTypes.h"
 #include "Misc/CoreMisc.h"
 #include "Subsystems/GameInstanceSubsystem.h"
 #include "Tickable.h"
diff --git a/Engine/Plugins/Experimental/GameplayTargetingSystem/Source/GameplayTargetingSystem/Tasks/TargetingFilterTask_ActorClass.h b/Engine/Plugins/Experimental/GameplayTargetingSystem/Source/GameplayTargetingSystem/Tasks/TargetingFilterTask_ActorClass.h
--- a/Engine/Plugins/Experimental/GameplayTargetingSystem/Source/GameplayTargetingSystem/Tasks/TargetingFilterTask_ActorClass.h
+++ b/Engine/Plugins/Experimental/GameplayTargetingSystem/Source/GameplayTargetingSystem/Tasks/TargetingFilterTask_ActorClass.h
@@ -1,8 +1,8 @@
 // Copyright Epic Games, Inc. All Rights Reserved.
 #pragma once
 
-#include "GameplayTargetingSystem/Tasks/TargetingFilterTask_BasicFilterTemplate.h"
-#include "GameplayTargetingSystem/Types/TargetingSystemTypes.h"
+#include "Tasks/TargetingFilterTask_BasicFilterTemplate.h"
+#include "Types/TargetingSystemTypes.h"
 
 #include "TargetingFilterTask_ActorClass.generated.h"
 
diff --git a/Engine/Plugins/Experimental/GameplayTargetingSystem/Source/GameplayTargetingSystem/Tasks/TargetingFilterTask_BasicFilterTemplate.cpp b/Engine/Plugins/Experimental/GameplayTargetingSystem/Source/GameplayTargetingSystem/Tasks/TargetingFilterTask_BasicFilterTemplate.cpp
--- a/Engine/Plugins/Experimental/GameplayTargetingSystem/Source/GameplayTargetingSystem/Tasks/TargetingFilterTask_BasicFilterTemplate.cpp
+++ b/Engine/Plugins/Experimental/GameplayTargetingSystem/Source/GameplayTargetingSystem/Tasks/TargetingFilterTask_BasicFilterTemplate.cpp
@@ -2,8 +2,8 @@
 #include "TargetingFilterTask_BasicFilterTemplate.h"
 
 #include "GameFramework/Actor.h"
-#include "GameplayTargetingSystem/TargetingSystem/TargetingSubsystem.h"
-#include "GameplayTargetingSystem/Types/TargetingSystemTypes.h"
+#include "TargetingSystem/TargetingSubsystem.h"
+#include "Types/TargetingSystemTypes.h"
 
 #if ENABLE_DRAW_DEBUG
 #include "Engine/Canvas.h"
diff --git a/Engine/Plugins/Experimental/GameplayTargetingSystem/Source/GameplayTargetingSystem/Tasks/TargetingFilterTask_SortByDistance.h b/Engine/Plugins/Experimental/GameplayTargetingSystem/Source/GameplayTargetingSystem/Tasks/TargetingFilterTask_SortByDistance.h
--- a/Engine/Plugins/Experimental/GameplayTargetingSystem/Source/GameplayTargetingSystem/Tasks/TargetingFilterTask_SortByDistance.h
+++ b/Engine/Plugins/Experimental/GameplayTargetingSystem/Source/GameplayTargetingSystem/Tasks/TargetingFilterTask_SortByDistance.h
@@ -1,7 +1,7 @@
 // Copyright Epic Games, Inc. All Rights Reserved.
 #pragma once
 
-#include "GameplayTargetingSystem/Types/TargetingSystemTypes.h"
+#include "Types/TargetingSystemTypes.h"
 #include "TargetingTask.h"
 #include "UObject/Object.h"
 #include "TargetingSortTask_Base.h"
diff --git a/Engine/Plugins/Experimental/GameplayTargetingSystem/Source/GameplayTargetingSystem/Tasks/TargetingSelectionTask_AOE.cpp b/Engine/Plugins/Experimental/GameplayTargetingSystem/Source/GameplayTargetingSystem/Tasks/TargetingSelectionTask_AOE.cpp
--- a/Engine/Plugins/Experimental/GameplayTargetingSystem/Source/GameplayTargetingSystem/Tasks/TargetingSelectionTask_AOE.cpp
+++ b/Engine/Plugins/Experimental/GameplayTargetingSystem/Source/GameplayTargetingSystem/Tasks/TargetingSelectionTask_AOE.cpp
@@ -7,8 +7,8 @@
 #include "Engine/EngineTypes.h"
 #include "Engine/World.h"
 #include "GameFramework/Actor.h"
-#include "GameplayTargetingSystem/TargetingSystem/TargetingSubsystem.h"
-#include "GameplayTargetingSystem/Types/TargetingSystemLogs.h"
+#include "TargetingSystem/TargetingSubsystem.h"
+#include "Types/TargetingSystemLogs.h"
 
 #if ENABLE_DRAW_DEBUG
 #include "Engine/Canvas.h"
diff --git a/Engine/Plugins/Experimental/GameplayTargetingSystem/Source/GameplayTargetingSystem/Tasks/TargetingSelectionTask_AOE.h b/Engine/Plugins/Experimental/GameplayTargetingSystem/Source/GameplayTargetingSystem/Tasks/TargetingSelectionTask_AOE.h
--- a/Engine/Plugins/Experimental/GameplayTargetingSystem/Source/GameplayTargetingSystem/Tasks/TargetingSelectionTask_AOE.h
+++ b/Engine/Plugins/Experimental/GameplayTargetingSystem/Source/GameplayTargetingSystem/Tasks/TargetingSelectionTask_AOE.h
@@ -4,7 +4,7 @@
 #include "CollisionShape.h"
 #include "Engine/CollisionProfile.h"
 #include "Engine/EngineTypes.h"
-#include "GameplayTargetingSystem/Types/TargetingSystemTypes.h"
+#include "Types/TargetingSystemTypes.h"
 #include "ScalableFloat.h"
 #include "TargetingTask.h"
 #include "UObject/Object.h"
diff --git a/Engine/Plugins/Experimental/GameplayTargetingSystem/Source/GameplayTargetingSystem/Tasks/TargetingSelectionTask_SourceActor.cpp b/Engine/Plugins/Experimental/GameplayTargetingSystem/Source/GameplayTargetingSystem/Tasks/TargetingSelectionTask_SourceActor.cpp
--- a/Engine/Plugins/Experimental/GameplayTargetingSystem/Source/GameplayTargetingSystem/Tasks/TargetingSelectionTask_SourceActor.cpp
+++ b/Engine/Plugins/Experimental/GameplayTargetingSystem/Source/GameplayTargetingSystem/Tasks/TargetingSelectionTask_SourceActor.cpp
@@ -2,7 +2,7 @@
 #include "TargetingSelectionTask_SourceActor.h"
 
 #include "GameFramework/Actor.h"
-#include "GameplayTargetingSystem/TargetingSystem/TargetingSubsystem.h"
+#include "TargetingSystem/TargetingSubsystem.h"
 
 #if ENABLE_DRAW_DEBUG
 #include "Engine/Canvas.h"
diff --git a/Engine/Plugins/Experimental/GameplayTargetingSystem/Source/GameplayTargetingSystem/Tasks/TargetingSelectionTask_SourceActor.h b/Engine/Plugins/Experimental/GameplayTargetingSystem/Source/GameplayTargetingSystem/Tasks/TargetingSelectionTask_SourceActor.h
--- a/Engine/Plugins/Experimental/GameplayTargetingSystem/Source/GameplayTargetingSystem/Tasks/TargetingSelectionTask_SourceActor.h
+++ b/Engine/Plugins/Experimental/GameplayTargetingSystem/Source/GameplayTargetingSystem/Tasks/TargetingSelectionTask_SourceActor.h
@@ -1,7 +1,7 @@
 // Copyright Epic Games, Inc. All Rights Reserved.
 #pragma once
 
-#include "GameplayTargetingSystem/Types/TargetingSystemTypes.h"
+#include "Types/TargetingSystemTypes.h"
 #include "TargetingTask.h"
 #include "UObject/Object.h"
 
diff --git a/Engine/Plugins/Experimental/GameplayTargetingSystem/Source/GameplayTargetingSystem/Tasks/TargetingSelectionTask_Trace.cpp b/Engine/Plugins/Experimental/GameplayTargetingSystem/Source/GameplayTargetingSystem/Tasks/TargetingSelectionTask_Trace.cpp
--- a/Engine/Plugins/Experimental/GameplayTargetingSystem/Source/GameplayTargetingSystem/Tasks/TargetingSelectionTask_Trace.cpp
+++ b/Engine/Plugins/Experimental/GameplayTargetingSystem/Source/GameplayTargetingSystem/Tasks/TargetingSelectionTask_Trace.cpp
@@ -7,8 +7,8 @@
 #include "Engine/World.h"
 #include "GameFramework/Actor.h"
 #include "GameFramework/Pawn.h"
-#include "GameplayTargetingSystem/TargetingSystem/TargetingSubsystem.h"
-#include "GameplayTargetingSystem/Types/TargetingSystemLogs.h"
+#include "TargetingSystem/TargetingSubsystem.h"
+#include "Types/TargetingSystemLogs.h"
 
 #if ENABLE_DRAW_DEBUG
 #include "Engine/Canvas.h"
diff --git a/Engine/Plugins/Experimental/GameplayTargetingSystem/Source/GameplayTargetingSystem/Tasks/TargetingSelectionTask_Trace.h b/Engine/Plugins/Experimental/GameplayTargetingSystem/Source/GameplayTargetingSystem/Tasks/TargetingSelectionTask_Trace.h
--- a/Engine/Plugins/Experimental/GameplayTargetingSystem/Source/GameplayTargetingSystem/Tasks/TargetingSelectionTask_Trace.h
+++ b/Engine/Plugins/Experimental/GameplayTargetingSystem/Source/GameplayTargetingSystem/Tasks/TargetingSelectionTask_Trace.h
@@ -3,7 +3,7 @@
 
 #include "CollisionShape.h"
 #include "Engine/CollisionProfile.h"
-#include "GameplayTargetingSystem/Types/TargetingSystemTypes.h"
+#include "Types/TargetingSystemTypes.h"
 #include "ScalableFloat.h"
 #include "TargetingTask.h"
 #include "UObject/Object.h"
diff --git a/Engine/Plugins/Experimental/GameplayTargetingSystem/Source/GameplayTargetingSystem/Tasks/TargetingSortTask_Base.cpp b/Engine/Plugins/Experimental/GameplayTargetingSystem/Source/GameplayTargetingSystem/Tasks/TargetingSortTask_Base.cpp
--- a/Engine/Plugins/Experimental/GameplayTargetingSystem/Source/GameplayTargetingSystem/Tasks/TargetingSortTask_Base.cpp
+++ b/Engine/Plugins/Experimental/GameplayTargetingSystem/Source/GameplayTargetingSystem/Tasks/TargetingSortTask_Base.cpp
@@ -3,7 +3,7 @@
 #include "TargetingSortTask_Base.h"
 
 #include "GameFramework/Actor.h"
-#include "GameplayTargetingSystem/TargetingSystem/TargetingSubsystem.h"
+#include "TargetingSystem/TargetingSubsystem.h"
 #include "Kismet/KismetMathLibrary.h"
 
 #if ENABLE_DRAW_DEBUG
diff --git a/Engine/Plugins/Experimental/GameplayTargetingSystem/Source/GameplayTargetingSystem/Tasks/TargetingTask.cpp b/Engine/Plugins/Experimental/GameplayTargetingSystem/Source/GameplayTargetingSystem/Tasks/TargetingTask.cpp
--- a/Engine/Plugins/Experimental/GameplayTargetingSystem/Source/GameplayTargetingSystem/Tasks/TargetingTask.cpp
+++ b/Engine/Plugins/Experimental/GameplayTargetingSystem/Source/GameplayTargetingSystem/Tasks/TargetingTask.cpp
@@ -2,7 +2,7 @@
 #include "TargetingTask.h"
 
 #include "GameFramework/Actor.h"
-#include "GameplayTargetingSystem/Types/TargetingSystemTypes.h"
+#include "Types/TargetingSystemTypes.h"
 
 
 UTargetingTask::UTargetingTask(const FObjectInitializer& ObjectInitializer)
diff --git a/Engine/Plugins/Experimental/GameplayTargetingSystem/Source/GameplayTargetingSystem/Tasks/TargetingTask.h b/Engine/Plugins/Experimental/GameplayTargetingSystem/Source/GameplayTargetingSystem/Tasks/TargetingTask.h
--- a/Engine/Plugins/Experimental/GameplayTargetingSystem/Source/GameplayTargetingSystem/Tasks/TargetingTask.h
+++ b/Engine/Plugins/Experimental/GameplayTargetingSystem/Source/GameplayTargetingSystem/Tasks/TargetingTask.h
@@ -2,7 +2,7 @@
 #pragma once
 
 #include "DrawDebugHelpers.h"
-#include "GameplayTargetingSystem/Types/TargetingSystemTypes.h"
+#include "Types/TargetingSystemTypes.h"
 #include "UObject/Object.h"
 
 #include "TargetingTask.generated.h"
diff --git a/Engine/Plugins/Experimental/GameplayTargetingSystem/Source/GameplayTargetingSystem/Types/TargetingSystemDataStores.h b/Engine/Plugins/Experimental/GameplayTargetingSystem/Source/GameplayTargetingSystem/Types/TargetingSystemDataStores.h
--- a/Engine/Plugins/Experimental/GameplayTargetingSystem/Source/GameplayTargetingSystem/Types/TargetingSystemDataStores.h
+++ b/Engine/Plugins/Experimental/GameplayTargetingSystem/Source/GameplayTargetingSystem/Types/TargetingSystemDataStores.h
@@ -5,7 +5,7 @@
 #include "CoreMinimal.h"
 #include "Delegates/DelegateCombinations.h"
 #include "Engine/EngineTypes.h"
-#include "GameplayTargetingSystem/TargetingSystem/TargetingSubsystem.h"
+#include "TargetingSystem/TargetingSubsystem.h"
 #include "TargetingSystemTypes.h"
 
 
diff --git a/Engine/Plugins/Experimental/GameplayTargetingSystem/Source/GameplayTargetingSystem/Types/TargetingSystemTypes.cpp b/Engine/Plugins/Experimental/GameplayTargetingSystem/Source/GameplayTargetingSystem/Types/TargetingSystemTypes.cpp
--- a/Engine/Plugins/Experimental/GameplayTargetingSystem/Source/GameplayTargetingSystem/Types/TargetingSystemTypes.cpp
+++ b/Engine/Plugins/Experimental/GameplayTargetingSystem/Source/GameplayTargetingSystem/Types/TargetingSystemTypes.cpp
@@ -1,9 +1,9 @@
 // Copyright Epic Games, Inc. All Rights Reserved.
 #include "TargetingSystemTypes.h"
 
-#include "GameplayTargetingSystem/TargetingSystem/TargetingPreset.h"
-#include "GameplayTargetingSystem/TargetingSystem/TargetingSubsystem.h"
-#include "GameplayTargetingSystem/Tasks/TargetingTask.h"
+#include "TargetingSystem/TargetingPreset.h"
+#include "TargetingSystem/TargetingSubsystem.h"
+#include "Tasks/TargetingTask.h"
 #include "TargetingSystemLogs.h"
 #include "TargetingSystemDataStores.h"
 
