PR #9936: Fixed RawScores getting populated wrong. (Contributed by [REDACTED_PERSON])

diff --git a/Engine/Plugins/Experimental/GameplayTargetingSystem/Source/GameplayTargetingSystem/Tasks/TargetingSortTask_Base.cpp b/Engine/Plugins/Experimental/GameplayTargetingSystem/Source/GameplayTargetingSystem/Tasks/TargetingSortTask_Base.cpp
--- a/Engine/Plugins/Experimental/GameplayTargetingSystem/Source/GameplayTargetingSystem/Tasks/TargetingSortTask_Base.cpp
+++ b/Engine/Plugins/Experimental/GameplayTargetingSystem/Source/GameplayTargetingSystem/Tasks/TargetingSortTask_Base.cpp
@@ -51,10 +51,9 @@ void UTargetingSortTask_Base::Execute(const FTargetingRequestHandle& TargetingHa
 			// every task should have the same max score so none weights more than the others
 			float HighestScore = 0.f;
 			TArray<float> RawScores;
-			for (int32 TargetIterator = NumTargets - 1; TargetIterator >= 0; --TargetIterator)
+			for (const FTargetingDefaultResultData& TargetResult : ResultData->TargetResults)
 			{
-				FTargetingDefaultResultData& TargetResult = ResultData->TargetResults[TargetIterator];
-				float RawScore = GetScoreForTarget(TargetingHandle, TargetResult);
+				const float RawScore = GetScoreForTarget(TargetingHandle, TargetResult);
 				RawScores.Add(RawScore);
 				HighestScore = FMath::Max(HighestScore, RawScore);
 			}
@@ -62,7 +61,7 @@ void UTargetingSortTask_Base::Execute(const FTargetingRequestHandle& TargetingHa
 			if(ensureMsgf(NumTargets == RawScores.Num(), TEXT("The cached raw scores should be the same size as the number of targets!")))
 			{
 				// Adding the normalized scores to each target result.
-				for (int32 TargetIterator = NumTargets - 1; TargetIterator >= 0; --TargetIterator)
+				for (int32 TargetIterator = 0; TargetIterator < NumTargets; ++TargetIterator)
 				{
 					FTargetingDefaultResultData& TargetResult = ResultData->TargetResults[TargetIterator];
 					
