Skip TargetingTaskSet elements that are null

diff --git a/Engine/Plugins/Experimental/GameplayTargetingSystem/Source/GameplayTargetingSystem/TargetingSystem/TargetingSubsystem.cpp b/Engine/Plugins/Experimental/GameplayTargetingSystem/Source/GameplayTargetingSystem/TargetingSystem/TargetingSubsystem.cpp
--- a/Engine/Plugins/Experimental/GameplayTargetingSystem/Source/GameplayTargetingSystem/TargetingSystem/TargetingSubsystem.cpp
+++ b/Engine/Plugins/Experimental/GameplayTargetingSystem/Source/GameplayTargetingSystem/TargetingSystem/TargetingSubsystem.cpp
@@ -573,6 +573,15 @@ void UTargetingSubsystem::ProcessTargetingRequestTasks(FTargetingRequestHandle&
 							}
 						}
 					}
+					else
+					{
+						TARGETING_LOG(Warning, TEXT("UTargetingSubsystem::ProcessTargetingRequestTasks - Encountered an empty Task, skipping."));
+						++CurrentTaskIndex;
+						if (CurrentTaskIndex >= NumTasks)
+						{
+							bRequestComplete = true;
+						}
+					}
 				}
 
 				if (bRequestComplete)
