Mark UTargetingTask as Abstract so it can't be chosen for EditInlineNew properties

diff --git a/Engine/Plugins/Experimental/GameplayTargetingSystem/Source/GameplayTargetingSystem/Tasks/TargetingTask.h b/Engine/Plugins/Experimental/GameplayTargetingSystem/Source/GameplayTargetingSystem/Tasks/TargetingTask.h
--- a/Engine/Plugins/Experimental/GameplayTargetingSystem/Source/GameplayTargetingSystem/Tasks/TargetingTask.h
+++ b/Engine/Plugins/Experimental/GameplayTargetingSystem/Source/GameplayTargetingSystem/Tasks/TargetingTask.h
@@ -40,7 +40,7 @@ struct FTargetingDebugInfo;
 *	the end user might prefer to make decisions on. Distance (min/max), score rating
 *	etc.
 */
-UCLASS(EditInlineNew)
+UCLASS(EditInlineNew, Abstract)
 class TARGETINGSYSTEM_API UTargetingTask : public UObject
 {
 	GENERATED_BODY()
