Adding an option to AOE targeting task to offset based on relative space rather than world space.

#tests Tested by making the targeting frequency of a  projectile be 1s and turning on targeting debug drawing, the offset is indeed done in relative space when the option is turned on

diff --git a/Engine/Plugins/Experimental/GameplayTargetingSystem/Source/GameplayTargetingSystem/Tasks/TargetingSelectionTask_AOE.cpp b/Engine/Plugins/Experimental/GameplayTargetingSystem/Source/GameplayTargetingSystem/Tasks/TargetingSelectionTask_AOE.cpp
--- a/Engine/Plugins/Experimental/GameplayTargetingSystem/Source/GameplayTargetingSystem/Tasks/TargetingSelectionTask_AOE.cpp
+++ b/Engine/Plugins/Experimental/GameplayTargetingSystem/Source/GameplayTargetingSystem/Tasks/TargetingSelectionTask_AOE.cpp
@@ -21,6 +21,7 @@ UTargetingSelectionTask_AOE::UTargetingSelectionTask_AOE(const FObjectInitialize
 	CollisionChannel = ECC_PhysicsBody;
 	bIgnoreSourceActor = false;
 	bIgnoreInstigatorActor = false;
+	bUseRelativeOffset = false;
 }
 
 void UTargetingSelectionTask_AOE::Execute(const FTargetingRequestHandle& TargetingHandle) const
@@ -265,6 +266,18 @@ FVector UTargetingSelectionTask_AOE::GetSourceLocation_Implementation(const FTar
 
 FVector UTargetingSelectionTask_AOE::GetSourceOffset_Implementation(const FTargetingRequestHandle& TargetingHandle) const
 {
+	// If we want to apply a relative offset, we offset the different base vectors of the local space of the source actor with the values provided
+	if(bUseRelativeOffset)
+	{
+		const FTargetingSourceContext* SourceContext = FTargetingSourceContext::Find(TargetingHandle);
+		if (const AActor* SourceActor = SourceContext ? SourceContext->SourceActor : nullptr)
+		{
+			return	SourceActor->GetActorForwardVector() * DefaultSourceOffset.X +
+					SourceActor->GetActorRightVector() * DefaultSourceOffset.Y +
+					SourceActor->GetActorUpVector() * DefaultSourceOffset.Z;
+		}
+	}
+
 	return DefaultSourceOffset;
 }
 
diff --git a/Engine/Plugins/Experimental/GameplayTargetingSystem/Source/GameplayTargetingSystem/Tasks/TargetingSelectionTask_AOE.h b/Engine/Plugins/Experimental/GameplayTargetingSystem/Source/GameplayTargetingSystem/Tasks/TargetingSelectionTask_AOE.h
--- a/Engine/Plugins/Experimental/GameplayTargetingSystem/Source/GameplayTargetingSystem/Tasks/TargetingSelectionTask_AOE.h
+++ b/Engine/Plugins/Experimental/GameplayTargetingSystem/Source/GameplayTargetingSystem/Tasks/TargetingSelectionTask_AOE.h
@@ -106,6 +106,10 @@ protected:
 	/** The default source location offset used by GetSourceOffset */
 	UPROPERTY(EditAnywhere, Category = "Target AOE Selection | Data")
 	FVector DefaultSourceOffset = FVector::ZeroVector;
+
+	/** Should we offset based on world or relative Source object transform? */
+	UPROPERTY(EditAnywhere, Category = "Target AOE Selection | Data")
+	uint8 bUseRelativeOffset : 1;
 	
 	/** Indicates the trace should ignore the source actor */
 	UPROPERTY(EditAnywhere, Category = "Target AOE Selection | Data")
