Hook into the TargetingSystem for gathering actors & Debug draw improvements

* Using the targeting system plugin for environment sensing and filtering actors around the viewer. The targeting system allows area of effect requests based on a cube, sphere, capsule and cylinder volumes. It also allows users to filter by actor class or sort the resulting actors.
* A default pipeline for the targeting system has been implemented to keep the UI simple. For advanced users, we allow to use custom targeting preset assets. Enabling the given checkbox, will show a asset selection property and hide all properties belonging to the build-in default pipeline.
* Changed the world sub-system to hold a multi-map of targets sorted by their owning actor. This allows us to iterate through all sub-targets in constant time without recursive traversal.
* Debug draw option now automatically shows/hides sub-options to render the target gizmos, detailed info, target provider infos and allows to specify a color for the current and candidate target debug draw.
* Assigning a nice random color for each viewer's current and candidate target debug drawing.

diff --git a/Engine/Plugins/Experimental/GameplayTargetingSystem/Source/GameplayTargetingSystem/TargetingSystem/TargetingPreset.h b/Engine/Plugins/Experimental/GameplayTargetingSystem/Source/GameplayTargetingSystem/TargetingSystem/TargetingPreset.h
--- a/Engine/Plugins/Experimental/GameplayTargetingSystem/Source/GameplayTargetingSystem/TargetingSystem/TargetingPreset.h
+++ b/Engine/Plugins/Experimental/GameplayTargetingSystem/Source/GameplayTargetingSystem/TargetingSystem/TargetingPreset.h
@@ -22,7 +22,7 @@ struct FTargetingDefaultResultsSet;
 *	and sorting.
 */
 UCLASS(BlueprintType)
-class UTargetingPreset : public UDataAsset
+class TARGETINGSYSTEM_API UTargetingPreset : public UDataAsset
 {
 	GENERATED_BODY()
 
@@ -31,6 +31,7 @@ public:
 
 	/** Get the set of tasks defined for this targeting preset */
 	FORCEINLINE const FTargetingTaskSet* GetTargetingTaskSet() const { return &TargetingTaskSet; }
+	FORCEINLINE FTargetingTaskSet& GetMutableTargetingTaskSet() { return TargetingTaskSet; }
 
 private:
 	/** The tasks that make up this targeting preset */
diff --git a/Engine/Plugins/Experimental/GameplayTargetingSystem/Source/GameplayTargetingSystem/Tasks/TargetingSelectionTask_AOE.cpp b/Engine/Plugins/Experimental/GameplayTargetingSystem/Source/GameplayTargetingSystem/Tasks/TargetingSelectionTask_AOE.cpp
--- a/Engine/Plugins/Experimental/GameplayTargetingSystem/Source/GameplayTargetingSystem/Tasks/TargetingSelectionTask_AOE.cpp
+++ b/Engine/Plugins/Experimental/GameplayTargetingSystem/Source/GameplayTargetingSystem/Tasks/TargetingSelectionTask_AOE.cpp
@@ -48,7 +48,7 @@ void UTargetingSelectionTask_AOE::ExecuteImmediateTrace(const FTargetingRequestH
 	ResetDebugString(TargetingHandle);
 #endif // ENABLE_DRAW_DEBUG
 
-	UWorld* World = GetSourceContextWorld(TargetingHandle);
+	const UWorld* World = GetSourceContextWorld(TargetingHandle);
 	if (World && TargetingHandle.IsValid())
 	{
 		const FVector SourceLocation = GetSourceLocation(TargetingHandle) + GetSourceOffset(TargetingHandle);
@@ -69,7 +69,7 @@ void UTargetingSelectionTask_AOE::ExecuteImmediateTrace(const FTargetingRequestH
 		}
 		else
 		{
-			FCollisionShape CollisionShape = GetCollisionShape();
+			const FCollisionShape CollisionShape = GetCollisionShape();
 			FCollisionQueryParams OverlapParams(TEXT("UTargetingSelectionTask_AOE"), SCENE_QUERY_STAT_ONLY(UTargetingSelectionTask_AOE_Shape), false);
 			InitCollisionParams(TargetingHandle, OverlapParams);
 
@@ -96,22 +96,7 @@ void UTargetingSelectionTask_AOE::ExecuteImmediateTrace(const FTargetingRequestH
 #if ENABLE_DRAW_DEBUG
 			if (UTargetingSubsystem::IsTargetingDebugEnabled())
 			{
-				switch (ShapeType)
-				{
-				case ETargetingAOEShape::Box:
-					DrawDebugBox(World, SourceLocation, CollisionShape.GetExtent(), FQuat::Identity, FColor::Red, false, 30.0f, 0, 2.0f);
-					break;
-				case ETargetingAOEShape::Sphere:
-					DrawDebugSphere(World, SourceLocation, CollisionShape.GetSphereRadius(), 32, FColor::Red, false, 30.0f, 0, 2.0f);
-					break;
-				case ETargetingAOEShape::Capsule:
-					DrawDebugCapsule(World, SourceLocation, CollisionShape.GetCapsuleHalfHeight(), CollisionShape.GetCapsuleRadius(), FQuat::Identity, FColor::Red, false, 30.0f, 0, 2.0f);
-					break;
-				case ETargetingAOEShape::Cylinder:
-					DrawDebugBox(World, SourceLocation, CollisionShape.GetExtent(), FQuat::Identity, FColor::Red, false, 30.0f, 0, 1.0f);
-					DrawDebugCylinder(World, -1.0f * CollisionShape.GetExtent(), CollisionShape.GetExtent(), CollisionShape.GetExtent().X, 32, FColor::Yellow, false, 30.0f, 0, 2.0f);
-					break;
-				}
+				DebugDrawBoundingVolume(TargetingHandle, FColor::Red);
 			}
 #endif // ENABLE_DRAW_DEBUG
 		}
@@ -122,6 +107,60 @@ void UTargetingSelectionTask_AOE::ExecuteImmediateTrace(const FTargetingRequestH
 	SetTaskAsyncState(TargetingHandle, ETargetingTaskAsyncState::Completed);
 }
 
+void UTargetingSelectionTask_AOE::DebugDrawBoundingVolume(const FTargetingRequestHandle& TargetingHandle, const FColor& Color) const
+{
+#if ENABLE_DRAW_DEBUG
+	const UWorld* World = GetSourceContextWorld(TargetingHandle);
+	const FVector SourceLocation = GetSourceLocation(TargetingHandle) + GetSourceOffset(TargetingHandle);
+	const FCollisionShape CollisionShape = GetCollisionShape();
+
+	const bool bPersistentLines = false;
+	const float LifeTime = 0.0f;
+	const uint8 DepthPriority = 0;
+	const float Thickness = 2.0f;
+
+	switch (ShapeType)
+	{
+	case ETargetingAOEShape::Box:
+		DrawDebugBox(World, SourceLocation, CollisionShape.GetExtent(), FQuat::Identity,
+			Color, bPersistentLines, LifeTime, DepthPriority, Thickness);
+		break;
+	case ETargetingAOEShape::Sphere:
+		DrawDebugCapsule(World, SourceLocation, CollisionShape.GetSphereRadius(), CollisionShape.GetSphereRadius(), FQuat::Identity,
+			Color, bPersistentLines, LifeTime, DepthPriority, Thickness);
+		break;
+	case ETargetingAOEShape::Capsule:
+		DrawDebugCapsule(World, SourceLocation, CollisionShape.GetCapsuleHalfHeight(), CollisionShape.GetCapsuleRadius(), FQuat::Identity,
+			Color, bPersistentLines, LifeTime, DepthPriority, Thickness);
+		break;
+	case ETargetingAOEShape::Cylinder:
+		DrawDebugCylinder(World, -1.0f * CollisionShape.GetExtent(), CollisionShape.GetExtent(), CollisionShape.GetExtent().X, 32,
+			Color, bPersistentLines, LifeTime, DepthPriority, Thickness);
+		break;
+	}
+#endif // ENABLE_DRAW_DEBUG
+}
+
+void UTargetingSelectionTask_AOE::SetShapeType(ETargetingAOEShape InShapeType)
+{
+	ShapeType = InShapeType;
+}
+
+void UTargetingSelectionTask_AOE::SetHalfExtent(FVector InHalfExtent)
+{
+	HalfExtent = InHalfExtent;
+}
+
+void UTargetingSelectionTask_AOE::SetRadius(FScalableFloat InRadius)
+{
+	Radius = InRadius;
+}
+
+void UTargetingSelectionTask_AOE::SetHalfHeight(FScalableFloat InHalfHeight)
+{
+	HalfHeight = InHalfHeight;
+}
+
 void UTargetingSelectionTask_AOE::ExecuteAsyncTrace(const FTargetingRequestHandle& TargetingHandle) const
 {
 	UWorld* World = GetSourceContextWorld(TargetingHandle);
diff --git a/Engine/Plugins/Experimental/GameplayTargetingSystem/Source/GameplayTargetingSystem/Tasks/TargetingSelectionTask_AOE.h b/Engine/Plugins/Experimental/GameplayTargetingSystem/Source/GameplayTargetingSystem/Tasks/TargetingSelectionTask_AOE.h
--- a/Engine/Plugins/Experimental/GameplayTargetingSystem/Source/GameplayTargetingSystem/Tasks/TargetingSelectionTask_AOE.h
+++ b/Engine/Plugins/Experimental/GameplayTargetingSystem/Source/GameplayTargetingSystem/Tasks/TargetingSelectionTask_AOE.h
@@ -44,7 +44,7 @@ enum class ETargetingAOEShape : uint8
 *	SourceComponent - Use a collision component with a specific component tag as the shape
 */
 UCLASS(Blueprintable)
-class UTargetingSelectionTask_AOE : public UTargetingTask
+class TARGETINGSYSTEM_API UTargetingSelectionTask_AOE : public UTargetingTask
 {
 	GENERATED_BODY()
 
@@ -54,6 +54,14 @@ public:
 	/** Evaluation function called by derived classes to process the targeting request */
 	virtual void Execute(const FTargetingRequestHandle& TargetingHandle) const override;
 
+	/** Debug draws the outlines of the set shape type. */
+	void DebugDrawBoundingVolume(const FTargetingRequestHandle& TargetingHandle, const FColor& Color) const;
+
+	void SetShapeType(ETargetingAOEShape InShapeType);
+	void SetHalfExtent(FVector InHalfExtent);
+	void SetRadius(FScalableFloat InRadius);
+	void SetHalfHeight(FScalableFloat InHalfHeight);
+
 private:
 	/** Method to process the trace task immediately */
 	void ExecuteImmediateTrace(const FTargetingRequestHandle& TargetingHandle) const;
