TargetingSelectionTask_AOE should respect source actor rotation instead of being axis-aligned

diff --git a/Engine/Plugins/Experimental/GameplayTargetingSystem/Source/GameplayTargetingSystem/Tasks/TargetingSelectionTask_AOE.cpp b/Engine/Plugins/Experimental/GameplayTargetingSystem/Source/GameplayTargetingSystem/Tasks/TargetingSelectionTask_AOE.cpp
--- a/Engine/Plugins/Experimental/GameplayTargetingSystem/Source/GameplayTargetingSystem/Tasks/TargetingSelectionTask_AOE.cpp
+++ b/Engine/Plugins/Experimental/GameplayTargetingSystem/Source/GameplayTargetingSystem/Tasks/TargetingSelectionTask_AOE.cpp
@@ -52,6 +52,7 @@ void UTargetingSelectionTask_AOE::ExecuteImmediateTrace(const FTargetingRequestH
 	if (World && TargetingHandle.IsValid())
 	{
 		const FVector SourceLocation = GetSourceLocation(TargetingHandle) + GetSourceOffset(TargetingHandle);
+		const FQuat SourceRotation = GetSourceRotation(TargetingHandle);
 
 		TArray<FOverlapResult> OverlapResults;
 		if (ShapeType == ETargetingAOEShape::SourceComponent)
@@ -82,15 +83,15 @@ void UTargetingSelectionTask_AOE::ExecuteImmediateTrace(const FTargetingRequestH
 					ObjectParams.AddObjectTypesToQuery(Channel);
 				}
 
-				World->OverlapMultiByObjectType(OverlapResults, SourceLocation, FQuat::Identity, ObjectParams, CollisionShape, OverlapParams);
+				World->OverlapMultiByObjectType(OverlapResults, SourceLocation, SourceRotation, ObjectParams, CollisionShape, OverlapParams);
 			}
 			else if (CollisionProfileName.Name != TEXT("NoCollision"))
 			{
-				World->OverlapMultiByProfile(OverlapResults, SourceLocation, FQuat::Identity, CollisionProfileName.Name, CollisionShape, OverlapParams);
+				World->OverlapMultiByProfile(OverlapResults, SourceLocation, SourceRotation, CollisionProfileName.Name, CollisionShape, OverlapParams);
 			}
 			else
 			{
-				World->OverlapMultiByChannel(OverlapResults, SourceLocation, FQuat::Identity, CollisionChannel, CollisionShape, OverlapParams);
+				World->OverlapMultiByChannel(OverlapResults, SourceLocation, SourceRotation, CollisionChannel, CollisionShape, OverlapParams);
 			}
 
 #if ENABLE_DRAW_DEBUG
@@ -107,11 +108,12 @@ void UTargetingSelectionTask_AOE::ExecuteImmediateTrace(const FTargetingRequestH
 	SetTaskAsyncState(TargetingHandle, ETargetingTaskAsyncState::Completed);
 }
 
-void UTargetingSelectionTask_AOE::DebugDrawBoundingVolume(const FTargetingRequestHandle& TargetingHandle, const FColor& Color) const
+void UTargetingSelectionTask_AOE::DebugDrawBoundingVolume(const FTargetingRequestHandle& TargetingHandle, const FColor& Color, const FOverlapDatum* OverlapDatum) const
 {
 #if ENABLE_DRAW_DEBUG
 	const UWorld* World = GetSourceContextWorld(TargetingHandle);
-	const FVector SourceLocation = GetSourceLocation(TargetingHandle) + GetSourceOffset(TargetingHandle);
+	const FVector SourceLocation = OverlapDatum ? OverlapDatum->Pos : GetSourceLocation(TargetingHandle) + GetSourceOffset(TargetingHandle);
+	const FQuat SourceRotation = OverlapDatum ? OverlapDatum->Rot : GetSourceRotation(TargetingHandle);
 	const FCollisionShape CollisionShape = GetCollisionShape();
 
 	const bool bPersistentLines = false;
@@ -122,21 +124,24 @@ void UTargetingSelectionTask_AOE::DebugDrawBoundingVolume(const FTargetingReques
 	switch (ShapeType)
 	{
 	case ETargetingAOEShape::Box:
-		DrawDebugBox(World, SourceLocation, CollisionShape.GetExtent(), FQuat::Identity,
+		DrawDebugBox(World, SourceLocation, CollisionShape.GetExtent(), SourceRotation,
 			Color, bPersistentLines, LifeTime, DepthPriority, Thickness);
 		break;
 	case ETargetingAOEShape::Sphere:
-		DrawDebugCapsule(World, SourceLocation, CollisionShape.GetSphereRadius(), CollisionShape.GetSphereRadius(), FQuat::Identity,
+		DrawDebugCapsule(World, SourceLocation, CollisionShape.GetSphereRadius(), CollisionShape.GetSphereRadius(), SourceRotation,
 			Color, bPersistentLines, LifeTime, DepthPriority, Thickness);
 		break;
 	case ETargetingAOEShape::Capsule:
-		DrawDebugCapsule(World, SourceLocation, CollisionShape.GetCapsuleHalfHeight(), CollisionShape.GetCapsuleRadius(), FQuat::Identity,
+		DrawDebugCapsule(World, SourceLocation, CollisionShape.GetCapsuleHalfHeight(), CollisionShape.GetCapsuleRadius(), SourceRotation,
 			Color, bPersistentLines, LifeTime, DepthPriority, Thickness);
 		break;
 	case ETargetingAOEShape::Cylinder:
-		DrawDebugCylinder(World, -1.0f * CollisionShape.GetExtent(), CollisionShape.GetExtent(), CollisionShape.GetExtent().X, 32,
-			Color, bPersistentLines, LifeTime, DepthPriority, Thickness);
-		break;
+		{
+			const FVector RotatedExtent = SourceRotation * CollisionShape.GetExtent();
+			DrawDebugCylinder(World, SourceLocation - RotatedExtent, SourceLocation + RotatedExtent, CollisionShape.GetExtent().X, 32,
+				Color, bPersistentLines, LifeTime, DepthPriority, Thickness);
+			break;
+		}
 	}
 #endif // ENABLE_DRAW_DEBUG
 }
@@ -167,6 +172,7 @@ void UTargetingSelectionTask_AOE::ExecuteAsyncTrace(const FTargetingRequestHandl
 	if (World && TargetingHandle.IsValid())
 	{
 		const FVector SourceLocation = GetSourceLocation(TargetingHandle) + GetSourceOffset(TargetingHandle);
+		const FQuat SourceRotation = GetSourceRotation(TargetingHandle);
 
 		FCollisionShape CollisionShape = GetCollisionShape();
 		FCollisionQueryParams OverlapParams(TEXT("UTargetingSelectionTask_AOE"), SCENE_QUERY_STAT_ONLY(UTargetingSelectionTask_AOE_Shape), false);
@@ -182,15 +188,15 @@ void UTargetingSelectionTask_AOE::ExecuteAsyncTrace(const FTargetingRequestHandl
 				ObjectParams.AddObjectTypesToQuery(Channel);
 			}
 
-			World->AsyncOverlapByObjectType(SourceLocation, FQuat::Identity, ObjectParams, CollisionShape, OverlapParams, &Delegate);
+			World->AsyncOverlapByObjectType(SourceLocation, SourceRotation, ObjectParams, CollisionShape, OverlapParams, &Delegate);
 		}
 		else if (CollisionProfileName.Name != TEXT("NoCollision"))
 		{
-			World->AsyncOverlapByProfile(SourceLocation, FQuat::Identity, CollisionProfileName.Name, CollisionShape, OverlapParams, &Delegate);
+			World->AsyncOverlapByProfile(SourceLocation, SourceRotation, CollisionProfileName.Name, CollisionShape, OverlapParams, &Delegate);
 		}
 		else
 		{
-			World->AsyncOverlapByChannel(SourceLocation, FQuat::Identity, CollisionChannel, CollisionShape, OverlapParams, FCollisionResponseParams::DefaultResponseParam, &Delegate);
+			World->AsyncOverlapByChannel(SourceLocation, SourceRotation, CollisionChannel, CollisionShape, OverlapParams, FCollisionResponseParams::DefaultResponseParam, &Delegate);
 		}
 	}
 	else
@@ -208,27 +214,7 @@ void UTargetingSelectionTask_AOE::HandleAsyncOverlapComplete(const FTraceHandle&
 
 		if (UTargetingSubsystem::IsTargetingDebugEnabled())
 		{
-			if (UWorld* World = GetSourceContextWorld(TargetingHandle))
-			{
-				const FVector& SourceLocation = InOverlapDatum.Pos;
-				const FCollisionShape& CollisionShape = InOverlapDatum.CollisionParams.CollisionShape;
-				switch (ShapeType)
-				{
-				case ETargetingAOEShape::Box:
-					DrawDebugBox(World, SourceLocation, CollisionShape.GetExtent(), FQuat::Identity, FColor::Red, false, 30.0f, 0, 2.0f);
-					break;
-				case ETargetingAOEShape::Sphere:
-					DrawDebugSphere(World, SourceLocation, CollisionShape.GetSphereRadius(), 32, FColor::Red, false, 30.0f, 0, 2.0f);
-					break;
-				case ETargetingAOEShape::Capsule:
-					DrawDebugCapsule(World, SourceLocation, CollisionShape.GetCapsuleHalfHeight(), CollisionShape.GetCapsuleRadius(), FQuat::Identity, FColor::Red, false, 30.0f, 0, 2.0f);
-					break;
-				case ETargetingAOEShape::Cylinder:
-					DrawDebugBox(World, SourceLocation, CollisionShape.GetExtent(), FQuat::Identity, FColor::Red, false, 30.0f, 0, 1.0f);
-					DrawDebugCylinder(World, -1.0f * CollisionShape.GetExtent(), CollisionShape.GetExtent(), CollisionShape.GetExtent().X, 32, FColor::Yellow, false, 30.0f, 0, 2.0f);
-					break;
-				}
-			}
+			DebugDrawBoundingVolume(TargetingHandle, FColor::Red, &InOverlapDatum);
 		}
 #endif // ENABLE_DRAW_DEBUG
 
@@ -320,6 +306,19 @@ FVector UTargetingSelectionTask_AOE::GetSourceOffset_Implementation(const FTarge
 	return DefaultSourceOffset;
 }
 
+FQuat UTargetingSelectionTask_AOE::GetSourceRotation_Implementation(const FTargetingRequestHandle& TargetingHandle) const
+{
+	if (const FTargetingSourceContext* SourceContext = FTargetingSourceContext::Find(TargetingHandle))
+	{
+		if (SourceContext->SourceActor)
+		{
+			return SourceContext->SourceActor->GetActorQuat();
+		}
+	}
+
+	return FQuat::Identity;
+}
+
 FCollisionShape UTargetingSelectionTask_AOE::GetCollisionShape() const
 {
 	switch (ShapeType)
diff --git a/Engine/Plugins/Experimental/GameplayTargetingSystem/Source/GameplayTargetingSystem/Tasks/TargetingSelectionTask_AOE.h b/Engine/Plugins/Experimental/GameplayTargetingSystem/Source/GameplayTargetingSystem/Tasks/TargetingSelectionTask_AOE.h
--- a/Engine/Plugins/Experimental/GameplayTargetingSystem/Source/GameplayTargetingSystem/Tasks/TargetingSelectionTask_AOE.h
+++ b/Engine/Plugins/Experimental/GameplayTargetingSystem/Source/GameplayTargetingSystem/Tasks/TargetingSelectionTask_AOE.h
@@ -55,7 +55,7 @@ public:
 	virtual void Execute(const FTargetingRequestHandle& TargetingHandle) const override;
 
 	/** Debug draws the outlines of the set shape type. */
-	void DebugDrawBoundingVolume(const FTargetingRequestHandle& TargetingHandle, const FColor& Color) const;
+	void DebugDrawBoundingVolume(const FTargetingRequestHandle& TargetingHandle, const FColor& Color, const FOverlapDatum* OverlapDatum = nullptr) const;
 
 	void SetShapeType(ETargetingAOEShape InShapeType);
 	void SetHalfExtent(FVector InHalfExtent);
@@ -84,6 +84,10 @@ protected:
 	UFUNCTION(BlueprintNativeEvent, Category = "Target AOE Selection")
 	FVector GetSourceOffset(const FTargetingRequestHandle& TargetingHandle) const;
 
+	/** Native event to get the source rotation for the AOE  */
+	UFUNCTION(BlueprintNativeEvent, Category = "Target AOE Selection")
+	FQuat GetSourceRotation(const FTargetingRequestHandle& TargetingHandle) const;
+	
 private:
 	/** Helper method to build the Collision Shape */
 	FCollisionShape GetCollisionShape() const;
