[Targeting System] Adding Data Store for Collision Params used by the AOE targeting task.
- Ading hook to override collision data via data store in the AOE targeting task.
- Exposing Async Task for targeting node so we can get the Targeting handle before we start targeting

diff --git a/Engine/Plugins/Experimental/GameplayTargetingSystem/Source/GameplayTargetingSystem/Async/AsyncAction_PerformTargeting.cpp b/Engine/Plugins/Experimental/GameplayTargetingSystem/Source/GameplayTargetingSystem/Async/AsyncAction_PerformTargeting.cpp
--- a/Engine/Plugins/Experimental/GameplayTargetingSystem/Source/GameplayTargetingSystem/Async/AsyncAction_PerformTargeting.cpp
+++ b/Engine/Plugins/Experimental/GameplayTargetingSystem/Source/GameplayTargetingSystem/Async/AsyncAction_PerformTargeting.cpp
@@ -73,9 +73,9 @@ void UAsyncAction_PerformTargeting::Activate()
 
 			TARGETING_LOG(Verbose, TEXT("Source Actor: %s"), *GetNameSafe(SourceContext.SourceActor));
 
-			FTargetingRequestHandle TargetingHandle = UTargetingSubsystem::MakeTargetRequestHandle(TargetingPreset, SourceContext);
+			TargetingHandle = UTargetingSubsystem::MakeTargetRequestHandle(TargetingPreset, SourceContext);
 
-			SetupInitialTargetsForRequest(TargetingHandle);
+			SetupInitialTargetsForRequest();
 
 			FTargetingRequestDelegate Delegate = FTargetingRequestDelegate::CreateWeakLambda(this, [this](FTargetingRequestHandle TargetingHandle)
 				{
@@ -112,7 +112,7 @@ void UAsyncAction_PerformTargeting::Activate()
 	}
 }
 
-void UAsyncAction_PerformTargeting::SetupInitialTargetsForRequest(const FTargetingRequestHandle& TargetingHandle) const
+void UAsyncAction_PerformTargeting::SetupInitialTargetsForRequest() const
 {
 	if (TargetingHandle.IsValid() && InitialTargets.Num() > 0)
 	{
diff --git a/Engine/Plugins/Experimental/GameplayTargetingSystem/Source/GameplayTargetingSystem/Async/AsyncAction_PerformTargeting.h b/Engine/Plugins/Experimental/GameplayTargetingSystem/Source/GameplayTargetingSystem/Async/AsyncAction_PerformTargeting.h
--- a/Engine/Plugins/Experimental/GameplayTargetingSystem/Source/GameplayTargetingSystem/Async/AsyncAction_PerformTargeting.h
+++ b/Engine/Plugins/Experimental/GameplayTargetingSystem/Source/GameplayTargetingSystem/Async/AsyncAction_PerformTargeting.h
@@ -18,7 +18,7 @@ DECLARE_DYNAMIC_MULTICAST_DELEGATE_OneParam(FPerformTargetingReady, FTargetingRe
 /**
 *	@class UAsyncAction_PerformTargeting
 */
-UCLASS()
+UCLASS(meta = (ExposedAsyncProxy = "AsyncTaskRef"))
 class UAsyncAction_PerformTargeting : public UBlueprintAsyncActionBase
 {
 	GENERATED_UCLASS_BODY()
@@ -30,6 +30,9 @@ public:
 	UFUNCTION(BlueprintCallable, meta = (DefaultToSelf = "SourceActor", BlueprintInternalUseOnly = "true", DisplayName = "Perform Filtering Async Action"))
 	static UAsyncAction_PerformTargeting* PerformFilteringRequest(AActor* SourceActor, UTargetingPreset* TargetingPreset, bool bUseAsyncTargeting, const TArray<AActor*> InTargets);
 
+	UFUNCTION(BlueprintPure, Category = "Targeting")
+	const FTargetingRequestHandle& GetTargetingHandle() const { return TargetingHandle; }
+
 	virtual void Activate() override;
 
 public:
@@ -38,7 +41,7 @@ public:
 
 private:
 	/** Method to seed the targeting request w/ the initial set of targets */
-	void SetupInitialTargetsForRequest(const FTargetingRequestHandle& TargetingHandle) const;
+	void SetupInitialTargetsForRequest() const;
 
 private:
 	/** The targeting preset to use for targeting */
@@ -53,6 +56,10 @@ private:
 	UPROPERTY()
 	TArray<TObjectPtr<AActor>> InitialTargets;
 
+	/** Targeting Handle created for this Async targeting task */
+	UPROPERTY()
+	FTargetingRequestHandle TargetingHandle;
+
 	/** Flag to indicate we should be using async targeting (default is immediate execution) */
 	UPROPERTY()
 	uint8 bUseAsyncTargeting : 1;
diff --git a/Engine/Plugins/Experimental/GameplayTargetingSystem/Source/GameplayTargetingSystem/TargetingSystem/TargetingSubsystem.cpp b/Engine/Plugins/Experimental/GameplayTargetingSystem/Source/GameplayTargetingSystem/TargetingSystem/TargetingSubsystem.cpp
--- a/Engine/Plugins/Experimental/GameplayTargetingSystem/Source/GameplayTargetingSystem/TargetingSystem/TargetingSubsystem.cpp
+++ b/Engine/Plugins/Experimental/GameplayTargetingSystem/Source/GameplayTargetingSystem/TargetingSystem/TargetingSubsystem.cpp
@@ -4,13 +4,14 @@
 #include "Engine/Engine.h"
 #include "Engine/GameInstance.h"
 #include "Engine/World.h"
-#include "Types/TargetingSystemLogs.h"
-#include "Types/TargetingSystemTypes.h"
 #include "HAL/IConsoleManager.h"
 #include "Kismet/KismetStringLibrary.h"
 #include "ProfilingDebugging/CsvProfiler.h"
 #include "Stats/Stats2.h"
 #include "TargetingPreset.h"
+#include "Tasks/CollisionQueryTaskData.h"
+#include "Types/TargetingSystemLogs.h"
+#include "Types/TargetingSystemTypes.h"
 
 #if ENABLE_DRAW_DEBUG
 #include "GameFramework/HUD.h"
@@ -121,6 +122,11 @@ void UTargetingSubsystem::AddReferencedObjects(UObject* InThis, FReferenceCollec
 				Collector.AddReferencedObjects<UTargetingTask>(FoundTaskSet->Tasks);
 			}
 		}
+		
+		if (FCollisionQueryTaskData* FoundOverride = UE::TargetingSystem::TTargetingDataStore<FCollisionQueryTaskData>::Find(CurHandle))
+		{
+			FoundOverride->AddStructReferencedObjects(Collector);
+		}
 	}
 }
 
@@ -649,6 +655,15 @@ void UTargetingSubsystem::GetTargetingResults(FTargetingRequestHandle TargetingH
 	}
 }
 
+void UTargetingSubsystem::OverrideCollisionQueryTaskData(const FTargetingRequestHandle& TargetingHandle, const FCollisionQueryTaskData& CollisionQueryDataOverride)
+{
+	if (TargetingHandle.IsValid())
+	{
+		FCollisionQueryTaskData& AddedCollisionQueryDataOverride = UE::TargetingSystem::TTargetingDataStore<FCollisionQueryTaskData>::FindOrAdd(TargetingHandle);
+		AddedCollisionQueryDataOverride = CollisionQueryDataOverride;
+	}
+}
+
 UTargetingTask* UTargetingSubsystem::FindCurrentExecutingTask(FTargetingRequestHandle Handle) const
 {
 	if (Handle.IsValid())
diff --git a/Engine/Plugins/Experimental/GameplayTargetingSystem/Source/GameplayTargetingSystem/TargetingSystem/TargetingSubsystem.h b/Engine/Plugins/Experimental/GameplayTargetingSystem/Source/GameplayTargetingSystem/TargetingSystem/TargetingSubsystem.h
--- a/Engine/Plugins/Experimental/GameplayTargetingSystem/Source/GameplayTargetingSystem/TargetingSystem/TargetingSubsystem.h
+++ b/Engine/Plugins/Experimental/GameplayTargetingSystem/Source/GameplayTargetingSystem/TargetingSystem/TargetingSubsystem.h
@@ -18,6 +18,7 @@ class UCanvas;
 class UTargetingPreset;
 class UTargetingTask;
 class UWorld;
+struct FCollisionQueryTaskData;
 
 
 #if ENABLE_DRAW_DEBUG
@@ -122,7 +123,7 @@ public:
 	TARGETINGSYSTEM_API static void ReleaseTargetRequestHandle(FTargetingRequestHandle& Handle);
 
 	/** The handle released delegate that fires right before a handle is release so all data stores can clean up their state */
-	DECLARE_MULTICAST_DELEGATE_OneParam(FOnTargetingRequestHandleReleased, FTargetingRequestHandle)
+	DECLARE_MULTICAST_DELEGATE_OneParam(FOnTargetingRequestHandleReleased, const FTargetingRequestHandle&)
 	static FOnTargetingRequestHandleReleased ReleaseHandleDelegate;
 
 	/** Target Handle Generation Methods */
@@ -183,6 +184,10 @@ public:
 	UFUNCTION(BlueprintCallable, Category = "Targeting System | Targeting Results")
 	TARGETINGSYSTEM_API void GetTargetingResults(FTargetingRequestHandle TargetingHandle, TArray<FHitResult>& OutTargets) const;
 
+	/** Function that lets you set a data store from a certain Targeting Handle to add some Collision Query Param Overrides  */
+	UFUNCTION(BlueprintCallable, Category = "Targeting System | Data Stores")
+	static void OverrideCollisionQueryTaskData(const FTargetingRequestHandle& TargetingHandle, const FCollisionQueryTaskData& CollisionQueryDataOverride);
+
 	/** ~Blueprint Helper Methods */
 
 	/** Debug Helper Methods */
diff --git a/Engine/Plugins/Experimental/GameplayTargetingSystem/Source/GameplayTargetingSystem/Tasks/CollisionQueryTaskData.cpp b/Engine/Plugins/Experimental/GameplayTargetingSystem/Source/GameplayTargetingSystem/Tasks/CollisionQueryTaskData.cpp
new file mode 100644
--- /dev/null
+++ b/Engine/Plugins/Experimental/GameplayTargetingSystem/Source/GameplayTargetingSystem/Tasks/CollisionQueryTaskData.cpp
@@ -0,0 +1,15 @@
+// Copyright Epic Games, Inc. All Rights Reserved.
+
+#pragma once
+
+#include "CollisionQueryTaskData.h"
+
+DEFINE_TARGETING_DATA_STORE(FCollisionQueryTaskData)
+
+void FCollisionQueryTaskData::AddStructReferencedObjects(class FReferenceCollector& Collector)
+{
+	for (TObjectPtr<AActor>& ReferencedActor : IgnoredActors)
+	{
+		Collector.AddReferencedObject(ReferencedActor);
+	}
+}
diff --git a/Engine/Plugins/Experimental/GameplayTargetingSystem/Source/GameplayTargetingSystem/Tasks/CollisionQueryTaskData.h b/Engine/Plugins/Experimental/GameplayTargetingSystem/Source/GameplayTargetingSystem/Tasks/CollisionQueryTaskData.h
new file mode 100644
--- /dev/null
+++ b/Engine/Plugins/Experimental/GameplayTargetingSystem/Source/GameplayTargetingSystem/Tasks/CollisionQueryTaskData.h
@@ -0,0 +1,25 @@
+// Copyright Epic Games, Inc. All Rights Reserved.
+
+#pragma once
+
+#include "Types/TargetingSystemDataStores.h"
+
+#include "CollisionQueryTaskData.generated.h"
+
+/** Data Store struct used to extend collision-based targeting tasks providing extra data from outside */
+USTRUCT(BlueprintType)
+struct FCollisionQueryTaskData 
+{
+	GENERATED_BODY()
+	
+	FCollisionQueryTaskData(){}
+
+	/** Any extra actors we want to ignore. Note: Given that this is a globally-managed struct, we're manually adding refs to it in UTargetingSubsystem::AddReferencedObjects */
+	UPROPERTY(BlueprintReadWrite, Category =  "Collision Query Task Overrides")
+	TArray<TObjectPtr<AActor>> IgnoredActors;
+
+	/** Called by the Targeting Subsystem to manually hang onto reference to object pointers in this struct given that it's globally-owned */
+	void AddStructReferencedObjects(class FReferenceCollector& Collector);
+};
+
+DECLARE_TARGETING_DATA_STORE(FCollisionQueryTaskData)
\ No newline at end of file
diff --git a/Engine/Plugins/Experimental/GameplayTargetingSystem/Source/GameplayTargetingSystem/Tasks/TargetingSelectionTask_AOE.cpp b/Engine/Plugins/Experimental/GameplayTargetingSystem/Source/GameplayTargetingSystem/Tasks/TargetingSelectionTask_AOE.cpp
--- a/Engine/Plugins/Experimental/GameplayTargetingSystem/Source/GameplayTargetingSystem/Tasks/TargetingSelectionTask_AOE.cpp
+++ b/Engine/Plugins/Experimental/GameplayTargetingSystem/Source/GameplayTargetingSystem/Tasks/TargetingSelectionTask_AOE.cpp
@@ -2,19 +2,20 @@
 #include "TargetingSelectionTask_AOE.h"
 
 #include "CollisionQueryParams.h"
+#include "CollisionQueryTaskData.h"
 #include "Components/PrimitiveComponent.h"
 #include "Engine/CollisionProfile.h"
 #include "Engine/EngineTypes.h"
 #include "Engine/World.h"
 #include "GameFramework/Actor.h"
 #include "TargetingSystem/TargetingSubsystem.h"
+#include "Types/TargetingSystemDataStores.h"
 #include "Types/TargetingSystemLogs.h"
 
 #if ENABLE_DRAW_DEBUG
 #include "Engine/Canvas.h"
 #endif // ENABLE_DRAW_DEBUG
 
-
 UTargetingSelectionTask_AOE::UTargetingSelectionTask_AOE(const FObjectInitializer& ObjectInitializer)
 : Super(ObjectInitializer)
 {
@@ -375,6 +376,12 @@ void UTargetingSelectionTask_AOE::InitCollisionParams(const FTargetingRequestHan
 			OutParams.AddIgnoredActor(SourceContext->InstigatorActor);
 		}
 	}
+
+	// If we have a data store to override the collision query params from, add that to our collision params
+	if (FCollisionQueryTaskData* FoundOverride = UE::TargetingSystem::TTargetingDataStore<FCollisionQueryTaskData>::Find(TargetingHandle))
+	{
+		OutParams.AddIgnoredActors(FoundOverride->IgnoredActors);
+	}
 }
 
 #if WITH_EDITOR
diff --git a/Engine/Plugins/Experimental/GameplayTargetingSystem/Source/GameplayTargetingSystem/Types/TargetingSystemDataStores.h b/Engine/Plugins/Experimental/GameplayTargetingSystem/Source/GameplayTargetingSystem/Types/TargetingSystemDataStores.h
--- a/Engine/Plugins/Experimental/GameplayTargetingSystem/Source/GameplayTargetingSystem/Types/TargetingSystemDataStores.h
+++ b/Engine/Plugins/Experimental/GameplayTargetingSystem/Source/GameplayTargetingSystem/Types/TargetingSystemDataStores.h
@@ -39,9 +39,9 @@ namespace UE
 				}
 			}
 
-			static Type& FindOrAdd(FTargetingRequestHandle Handle);
-			static Type* Find(FTargetingRequestHandle Handle);
-			static void OnTargetingRequestHandleReleased(FTargetingRequestHandle Handle);
+			static Type& FindOrAdd(const FTargetingRequestHandle& Handle);
+			static Type* Find(const FTargetingRequestHandle& Handle);
+			static void OnTargetingRequestHandleReleased(const FTargetingRequestHandle& Handle);
 
 		private:
 			// @note: SortMap can be a concern if the Type is a large struct and/or there is high volume of
@@ -62,19 +62,19 @@ namespace UE
 		extern TTargetingDataStore<FTargetingRequestData> GTargetingRequestDataStore;
 
 		template<>
-		FORCEINLINE FTargetingRequestData& TTargetingDataStore<FTargetingRequestData>::FindOrAdd(FTargetingRequestHandle Handle)
+		FORCEINLINE FTargetingRequestData& TTargetingDataStore<FTargetingRequestData>::FindOrAdd(const FTargetingRequestHandle& Handle)
 		{
 			return GTargetingRequestDataStore.Items.FindOrAdd(Handle);
 		}
 
 		template<>
-		FORCEINLINE FTargetingRequestData* TTargetingDataStore<FTargetingRequestData>::Find(FTargetingRequestHandle Handle)
+		FORCEINLINE FTargetingRequestData* TTargetingDataStore<FTargetingRequestData>::Find(const FTargetingRequestHandle& Handle)
 		{
 			return GTargetingRequestDataStore.Items.Find(Handle);
 		}
 
 		template<>
-		FORCEINLINE void TTargetingDataStore<FTargetingRequestData>::OnTargetingRequestHandleReleased(FTargetingRequestHandle Handle)
+		FORCEINLINE void TTargetingDataStore<FTargetingRequestData>::OnTargetingRequestHandleReleased(const FTargetingRequestHandle& Handle)
 		{
 			GTargetingRequestDataStore.Items.Remove(Handle);
 		}
@@ -84,19 +84,19 @@ namespace UE
 		extern TTargetingDataStore<const FTargetingTaskSet*> GTargetingTaskSetDataStore;
 
 		template<>
-		FORCEINLINE const FTargetingTaskSet*& TTargetingDataStore<const FTargetingTaskSet*>::FindOrAdd(FTargetingRequestHandle Handle)
+		FORCEINLINE const FTargetingTaskSet*& TTargetingDataStore<const FTargetingTaskSet*>::FindOrAdd(const FTargetingRequestHandle& Handle)
 		{
 			return GTargetingTaskSetDataStore.Items.FindOrAdd(Handle);
 		}
 
 		template<>
-		FORCEINLINE const FTargetingTaskSet** TTargetingDataStore<const FTargetingTaskSet*>::Find(FTargetingRequestHandle Handle)
+		FORCEINLINE const FTargetingTaskSet** TTargetingDataStore<const FTargetingTaskSet*>::Find(const FTargetingRequestHandle& Handle)
 		{
 			return GTargetingTaskSetDataStore.Items.Find(Handle);
 		}
 
 		template<>
-		FORCEINLINE void TTargetingDataStore<const FTargetingTaskSet*>::OnTargetingRequestHandleReleased(FTargetingRequestHandle Handle)
+		FORCEINLINE void TTargetingDataStore<const FTargetingTaskSet*>::OnTargetingRequestHandleReleased(const FTargetingRequestHandle& Handle)
 		{
 			GTargetingTaskSetDataStore.Items.Remove(Handle);
 		}
@@ -106,19 +106,19 @@ namespace UE
 		extern TTargetingDataStore<FTargetingSourceContext> GTargetingSourceDataStore;
 
 		template<>
-		FORCEINLINE FTargetingSourceContext& TTargetingDataStore<FTargetingSourceContext>::FindOrAdd(FTargetingRequestHandle Handle)
+		FORCEINLINE FTargetingSourceContext& TTargetingDataStore<FTargetingSourceContext>::FindOrAdd(const FTargetingRequestHandle& Handle)
 		{
 			return GTargetingSourceDataStore.Items.FindOrAdd(Handle);
 		}
 
 		template<>
-		FORCEINLINE FTargetingSourceContext* TTargetingDataStore<FTargetingSourceContext>::Find(FTargetingRequestHandle Handle)
+		FORCEINLINE FTargetingSourceContext* TTargetingDataStore<FTargetingSourceContext>::Find(const FTargetingRequestHandle& Handle)
 		{
 			return GTargetingSourceDataStore.Items.Find(Handle);
 		}
 
 		template<>
-		FORCEINLINE void TTargetingDataStore<FTargetingSourceContext>::OnTargetingRequestHandleReleased(FTargetingRequestHandle Handle)
+		FORCEINLINE void TTargetingDataStore<FTargetingSourceContext>::OnTargetingRequestHandleReleased(const FTargetingRequestHandle& Handle)
 		{
 			GTargetingSourceDataStore.Items.Remove(Handle);
 		}
@@ -128,19 +128,19 @@ namespace UE
 		extern TTargetingDataStore<FTargetingDefaultResultsSet> GTargetingResultsDataStore;
 
 		template<>
-		FORCEINLINE FTargetingDefaultResultsSet& TTargetingDataStore<FTargetingDefaultResultsSet>::FindOrAdd(FTargetingRequestHandle Handle)
+		FORCEINLINE FTargetingDefaultResultsSet& TTargetingDataStore<FTargetingDefaultResultsSet>::FindOrAdd(const FTargetingRequestHandle& Handle)
 		{
 			return GTargetingResultsDataStore.Items.FindOrAdd(Handle);
 		}
 
 		template<>
-		FORCEINLINE FTargetingDefaultResultsSet* TTargetingDataStore<FTargetingDefaultResultsSet>::Find(FTargetingRequestHandle Handle)
+		FORCEINLINE FTargetingDefaultResultsSet* TTargetingDataStore<FTargetingDefaultResultsSet>::Find(const FTargetingRequestHandle& Handle)
 		{
 			return GTargetingResultsDataStore.Items.Find(Handle);
 		}
 
 		template<>
-		FORCEINLINE void TTargetingDataStore<FTargetingDefaultResultsSet>::OnTargetingRequestHandleReleased(FTargetingRequestHandle Handle)
+		FORCEINLINE void TTargetingDataStore<FTargetingDefaultResultsSet>::OnTargetingRequestHandleReleased(const FTargetingRequestHandle& Handle)
 		{
 			GTargetingResultsDataStore.Items.Remove(Handle);
 		}
@@ -150,19 +150,19 @@ namespace UE
 		extern TTargetingDataStore<FTargetingAsyncTaskData> GTargetingAsyncTaskDataStore;
 
 		template<>
-		FORCEINLINE FTargetingAsyncTaskData& TTargetingDataStore<FTargetingAsyncTaskData>::FindOrAdd(FTargetingRequestHandle Handle)
+		FORCEINLINE FTargetingAsyncTaskData& TTargetingDataStore<FTargetingAsyncTaskData>::FindOrAdd(const FTargetingRequestHandle& Handle)
 		{
 			return GTargetingAsyncTaskDataStore.Items.FindOrAdd(Handle);
 		}
 
 		template<>
-		FORCEINLINE FTargetingAsyncTaskData* TTargetingDataStore<FTargetingAsyncTaskData>::Find(FTargetingRequestHandle Handle)
+		FORCEINLINE FTargetingAsyncTaskData* TTargetingDataStore<FTargetingAsyncTaskData>::Find(const FTargetingRequestHandle& Handle)
 		{
 			return GTargetingAsyncTaskDataStore.Items.Find(Handle);
 		}
 
 		template<>
-		FORCEINLINE void TTargetingDataStore<FTargetingAsyncTaskData>::OnTargetingRequestHandleReleased(FTargetingRequestHandle Handle)
+		FORCEINLINE void TTargetingDataStore<FTargetingAsyncTaskData>::OnTargetingRequestHandleReleased(const FTargetingRequestHandle& Handle)
 		{
 			GTargetingAsyncTaskDataStore.Items.Remove(Handle);
 		}
@@ -171,19 +171,19 @@ namespace UE
 		extern TTargetingDataStore<FTargetingImmediateTaskData> GTargetingImmediateTaskDataStore;
 
 		template<>
-		FORCEINLINE FTargetingImmediateTaskData& TTargetingDataStore<FTargetingImmediateTaskData>::FindOrAdd(FTargetingRequestHandle Handle)
+		FORCEINLINE FTargetingImmediateTaskData& TTargetingDataStore<FTargetingImmediateTaskData>::FindOrAdd(const FTargetingRequestHandle& Handle)
 		{
 			return GTargetingImmediateTaskDataStore.Items.FindOrAdd(Handle);
 		}
 
 		template<>
-		FORCEINLINE FTargetingImmediateTaskData* TTargetingDataStore<FTargetingImmediateTaskData>::Find(FTargetingRequestHandle Handle)
+		FORCEINLINE FTargetingImmediateTaskData* TTargetingDataStore<FTargetingImmediateTaskData>::Find(const FTargetingRequestHandle& Handle)
 		{
 			return GTargetingImmediateTaskDataStore.Items.Find(Handle);
 		}
 
 		template<>
-		FORCEINLINE void TTargetingDataStore<FTargetingImmediateTaskData>::OnTargetingRequestHandleReleased(FTargetingRequestHandle Handle)
+		FORCEINLINE void TTargetingDataStore<FTargetingImmediateTaskData>::OnTargetingRequestHandleReleased(const FTargetingRequestHandle& Handle)
 		{
 			GTargetingImmediateTaskDataStore.Items.Remove(Handle);
 		}
@@ -194,19 +194,19 @@ namespace UE
 		extern TTargetingDataStore<FTargetingDebugData> GTargetingDebugDataStore;
 
 		template<>
-		FORCEINLINE FTargetingDebugData& TTargetingDataStore<FTargetingDebugData>::FindOrAdd(FTargetingRequestHandle Handle)
+		FORCEINLINE FTargetingDebugData& TTargetingDataStore<FTargetingDebugData>::FindOrAdd(const FTargetingRequestHandle& Handle)
 		{
 			return GTargetingDebugDataStore.Items.FindOrAdd(Handle);
 		}
 
 		template<>
-		FORCEINLINE FTargetingDebugData* TTargetingDataStore<FTargetingDebugData>::Find(FTargetingRequestHandle Handle)
+		FORCEINLINE FTargetingDebugData* TTargetingDataStore<FTargetingDebugData>::Find(const FTargetingRequestHandle& Handle)
 		{
 			return GTargetingDebugDataStore.Items.Find(Handle);
 		}
 
 		template<>
-		FORCEINLINE void TTargetingDataStore<FTargetingDebugData>::OnTargetingRequestHandleReleased(FTargetingRequestHandle Handle)
+		FORCEINLINE void TTargetingDataStore<FTargetingDebugData>::OnTargetingRequestHandleReleased(const FTargetingRequestHandle& Handle)
 		{
 			GTargetingDebugDataStore.Items.Remove(Handle);
 		}
@@ -244,19 +244,19 @@ namespace UE
 		extern TTargetingDataStore<DataType> GTargetingDataStore##DataType;\
 \
 		template<>\
-		FORCEINLINE DataType& TTargetingDataStore<DataType>::FindOrAdd(FTargetingRequestHandle Handle)\
+		FORCEINLINE DataType& TTargetingDataStore<DataType>::FindOrAdd(const FTargetingRequestHandle& Handle)\
 		{\
 			return GTargetingDataStore##DataType.Items.FindOrAdd(Handle);\
 		}\
 \
 		template<>\
-		FORCEINLINE DataType* TTargetingDataStore<DataType>::Find(FTargetingRequestHandle Handle)\
+		FORCEINLINE DataType* TTargetingDataStore<DataType>::Find(const FTargetingRequestHandle& Handle)\
 		{\
 			return GTargetingDataStore##DataType.Items.Find(Handle);\
 		}\
 \
 		template<>\
-		FORCEINLINE void TTargetingDataStore<DataType>::OnTargetingRequestHandleReleased(FTargetingRequestHandle Handle)\
+		FORCEINLINE void TTargetingDataStore<DataType>::OnTargetingRequestHandleReleased(const FTargetingRequestHandle& Handle)\
 		{\
 			GTargetingDataStore##DataType.Items.Remove(Handle);\
 		}\
