Fixing naming issue where local variable in lambda would shadow Targeting Handle member variable.

diff --git a/Engine/Plugins/Experimental/GameplayTargetingSystem/Source/GameplayTargetingSystem/Async/AsyncAction_PerformTargeting.cpp b/Engine/Plugins/Experimental/GameplayTargetingSystem/Source/GameplayTargetingSystem/Async/AsyncAction_PerformTargeting.cpp
--- a/Engine/Plugins/Experimental/GameplayTargetingSystem/Source/GameplayTargetingSystem/Async/AsyncAction_PerformTargeting.cpp
+++ b/Engine/Plugins/Experimental/GameplayTargetingSystem/Source/GameplayTargetingSystem/Async/AsyncAction_PerformTargeting.cpp
@@ -77,11 +77,11 @@ void UAsyncAction_PerformTargeting::Activate()
 
 			SetupInitialTargetsForRequest();
 
-			FTargetingRequestDelegate Delegate = FTargetingRequestDelegate::CreateWeakLambda(this, [this](FTargetingRequestHandle TargetingHandle)
+			FTargetingRequestDelegate Delegate = FTargetingRequestDelegate::CreateWeakLambda(this, [this](FTargetingRequestHandle InTargetingHandle)
 				{
 					TARGETING_LOG(Verbose, TEXT("Entering request lambda"));
 
-					Targeted.Broadcast(TargetingHandle);
+					Targeted.Broadcast(InTargetingHandle);
 				});
 
 			if (bUseAsyncTargeting)
