[Targeting System] Reverting const correctness changes as it's prefered to pass by value for smaller types.

diff --git a/Engine/Plugins/Experimental/GameplayTargetingSystem/Source/GameplayTargetingSystem/TargetingSystem/TargetingSubsystem.cpp b/Engine/Plugins/Experimental/GameplayTargetingSystem/Source/GameplayTargetingSystem/TargetingSystem/TargetingSubsystem.cpp
--- a/Engine/Plugins/Experimental/GameplayTargetingSystem/Source/GameplayTargetingSystem/TargetingSystem/TargetingSubsystem.cpp
+++ b/Engine/Plugins/Experimental/GameplayTargetingSystem/Source/GameplayTargetingSystem/TargetingSystem/TargetingSubsystem.cpp
@@ -655,7 +655,7 @@ void UTargetingSubsystem::GetTargetingResults(FTargetingRequestHandle TargetingH
 	}
 }
 
-void UTargetingSubsystem::OverrideCollisionQueryTaskData(const FTargetingRequestHandle& TargetingHandle, const FCollisionQueryTaskData& CollisionQueryDataOverride)
+void UTargetingSubsystem::OverrideCollisionQueryTaskData(FTargetingRequestHandle TargetingHandle, const FCollisionQueryTaskData& CollisionQueryDataOverride)
 {
 	if (TargetingHandle.IsValid())
 	{
diff --git a/Engine/Plugins/Experimental/GameplayTargetingSystem/Source/GameplayTargetingSystem/TargetingSystem/TargetingSubsystem.h b/Engine/Plugins/Experimental/GameplayTargetingSystem/Source/GameplayTargetingSystem/TargetingSystem/TargetingSubsystem.h
--- a/Engine/Plugins/Experimental/GameplayTargetingSystem/Source/GameplayTargetingSystem/TargetingSystem/TargetingSubsystem.h
+++ b/Engine/Plugins/Experimental/GameplayTargetingSystem/Source/GameplayTargetingSystem/TargetingSystem/TargetingSubsystem.h
@@ -123,7 +123,7 @@ public:
 	TARGETINGSYSTEM_API static void ReleaseTargetRequestHandle(FTargetingRequestHandle& Handle);
 
 	/** The handle released delegate that fires right before a handle is release so all data stores can clean up their state */
-	DECLARE_MULTICAST_DELEGATE_OneParam(FOnTargetingRequestHandleReleased, const FTargetingRequestHandle&)
+	DECLARE_MULTICAST_DELEGATE_OneParam(FOnTargetingRequestHandleReleased, FTargetingRequestHandle)
 	static FOnTargetingRequestHandleReleased ReleaseHandleDelegate;
 
 	/** Target Handle Generation Methods */
@@ -186,7 +186,7 @@ public:
 
 	/** Function that lets you set a data store from a certain Targeting Handle to add some Collision Query Param Overrides  */
 	UFUNCTION(BlueprintCallable, Category = "Targeting System | Data Stores")
-	static void OverrideCollisionQueryTaskData(const FTargetingRequestHandle& TargetingHandle, const FCollisionQueryTaskData& CollisionQueryDataOverride);
+	static void OverrideCollisionQueryTaskData(FTargetingRequestHandle TargetingHandle, const FCollisionQueryTaskData& CollisionQueryDataOverride);
 
 	/** ~Blueprint Helper Methods */
 
diff --git a/Engine/Plugins/Experimental/GameplayTargetingSystem/Source/GameplayTargetingSystem/Types/TargetingSystemDataStores.h b/Engine/Plugins/Experimental/GameplayTargetingSystem/Source/GameplayTargetingSystem/Types/TargetingSystemDataStores.h
--- a/Engine/Plugins/Experimental/GameplayTargetingSystem/Source/GameplayTargetingSystem/Types/TargetingSystemDataStores.h
+++ b/Engine/Plugins/Experimental/GameplayTargetingSystem/Source/GameplayTargetingSystem/Types/TargetingSystemDataStores.h
@@ -39,9 +39,9 @@ namespace UE
 				}
 			}
 
-			static Type& FindOrAdd(const FTargetingRequestHandle& Handle);
-			static Type* Find(const FTargetingRequestHandle& Handle);
-			static void OnTargetingRequestHandleReleased(const FTargetingRequestHandle& Handle);
+			static Type& FindOrAdd(FTargetingRequestHandle Handle);
+			static Type* Find(FTargetingRequestHandle Handle);
+			static void OnTargetingRequestHandleReleased(FTargetingRequestHandle Handle);
 
 		private:
 			// @note: SortMap can be a concern if the Type is a large struct and/or there is high volume of
@@ -62,19 +62,19 @@ namespace UE
 		extern TTargetingDataStore<FTargetingRequestData> GTargetingRequestDataStore;
 
 		template<>
-		FORCEINLINE FTargetingRequestData& TTargetingDataStore<FTargetingRequestData>::FindOrAdd(const FTargetingRequestHandle& Handle)
+		FORCEINLINE FTargetingRequestData& TTargetingDataStore<FTargetingRequestData>::FindOrAdd(FTargetingRequestHandle Handle)
 		{
 			return GTargetingRequestDataStore.Items.FindOrAdd(Handle);
 		}
 
 		template<>
-		FORCEINLINE FTargetingRequestData* TTargetingDataStore<FTargetingRequestData>::Find(const FTargetingRequestHandle& Handle)
+		FORCEINLINE FTargetingRequestData* TTargetingDataStore<FTargetingRequestData>::Find(FTargetingRequestHandle Handle)
 		{
 			return GTargetingRequestDataStore.Items.Find(Handle);
 		}
 
 		template<>
-		FORCEINLINE void TTargetingDataStore<FTargetingRequestData>::OnTargetingRequestHandleReleased(const FTargetingRequestHandle& Handle)
+		FORCEINLINE void TTargetingDataStore<FTargetingRequestData>::OnTargetingRequestHandleReleased(FTargetingRequestHandle Handle)
 		{
 			GTargetingRequestDataStore.Items.Remove(Handle);
 		}
@@ -84,19 +84,19 @@ namespace UE
 		extern TTargetingDataStore<const FTargetingTaskSet*> GTargetingTaskSetDataStore;
 
 		template<>
-		FORCEINLINE const FTargetingTaskSet*& TTargetingDataStore<const FTargetingTaskSet*>::FindOrAdd(const FTargetingRequestHandle& Handle)
+		FORCEINLINE const FTargetingTaskSet*& TTargetingDataStore<const FTargetingTaskSet*>::FindOrAdd(FTargetingRequestHandle Handle)
 		{
 			return GTargetingTaskSetDataStore.Items.FindOrAdd(Handle);
 		}
 
 		template<>
-		FORCEINLINE const FTargetingTaskSet** TTargetingDataStore<const FTargetingTaskSet*>::Find(const FTargetingRequestHandle& Handle)
+		FORCEINLINE const FTargetingTaskSet** TTargetingDataStore<const FTargetingTaskSet*>::Find(FTargetingRequestHandle Handle)
 		{
 			return GTargetingTaskSetDataStore.Items.Find(Handle);
 		}
 
 		template<>
-		FORCEINLINE void TTargetingDataStore<const FTargetingTaskSet*>::OnTargetingRequestHandleReleased(const FTargetingRequestHandle& Handle)
+		FORCEINLINE void TTargetingDataStore<const FTargetingTaskSet*>::OnTargetingRequestHandleReleased(FTargetingRequestHandle Handle)
 		{
 			GTargetingTaskSetDataStore.Items.Remove(Handle);
 		}
@@ -106,19 +106,19 @@ namespace UE
 		extern TTargetingDataStore<FTargetingSourceContext> GTargetingSourceDataStore;
 
 		template<>
-		FORCEINLINE FTargetingSourceContext& TTargetingDataStore<FTargetingSourceContext>::FindOrAdd(const FTargetingRequestHandle& Handle)
+		FORCEINLINE FTargetingSourceContext& TTargetingDataStore<FTargetingSourceContext>::FindOrAdd(FTargetingRequestHandle Handle)
 		{
 			return GTargetingSourceDataStore.Items.FindOrAdd(Handle);
 		}
 
 		template<>
-		FORCEINLINE FTargetingSourceContext* TTargetingDataStore<FTargetingSourceContext>::Find(const FTargetingRequestHandle& Handle)
+		FORCEINLINE FTargetingSourceContext* TTargetingDataStore<FTargetingSourceContext>::Find(FTargetingRequestHandle Handle)
 		{
 			return GTargetingSourceDataStore.Items.Find(Handle);
 		}
 
 		template<>
-		FORCEINLINE void TTargetingDataStore<FTargetingSourceContext>::OnTargetingRequestHandleReleased(const FTargetingRequestHandle& Handle)
+		FORCEINLINE void TTargetingDataStore<FTargetingSourceContext>::OnTargetingRequestHandleReleased(FTargetingRequestHandle Handle)
 		{
 			GTargetingSourceDataStore.Items.Remove(Handle);
 		}
@@ -128,19 +128,19 @@ namespace UE
 		extern TTargetingDataStore<FTargetingDefaultResultsSet> GTargetingResultsDataStore;
 
 		template<>
-		FORCEINLINE FTargetingDefaultResultsSet& TTargetingDataStore<FTargetingDefaultResultsSet>::FindOrAdd(const FTargetingRequestHandle& Handle)
+		FORCEINLINE FTargetingDefaultResultsSet& TTargetingDataStore<FTargetingDefaultResultsSet>::FindOrAdd(FTargetingRequestHandle Handle)
 		{
 			return GTargetingResultsDataStore.Items.FindOrAdd(Handle);
 		}
 
 		template<>
-		FORCEINLINE FTargetingDefaultResultsSet* TTargetingDataStore<FTargetingDefaultResultsSet>::Find(const FTargetingRequestHandle& Handle)
+		FORCEINLINE FTargetingDefaultResultsSet* TTargetingDataStore<FTargetingDefaultResultsSet>::Find(FTargetingRequestHandle Handle)
 		{
 			return GTargetingResultsDataStore.Items.Find(Handle);
 		}
 
 		template<>
-		FORCEINLINE void TTargetingDataStore<FTargetingDefaultResultsSet>::OnTargetingRequestHandleReleased(const FTargetingRequestHandle& Handle)
+		FORCEINLINE void TTargetingDataStore<FTargetingDefaultResultsSet>::OnTargetingRequestHandleReleased(FTargetingRequestHandle Handle)
 		{
 			GTargetingResultsDataStore.Items.Remove(Handle);
 		}
@@ -150,19 +150,19 @@ namespace UE
 		extern TTargetingDataStore<FTargetingAsyncTaskData> GTargetingAsyncTaskDataStore;
 
 		template<>
-		FORCEINLINE FTargetingAsyncTaskData& TTargetingDataStore<FTargetingAsyncTaskData>::FindOrAdd(const FTargetingRequestHandle& Handle)
+		FORCEINLINE FTargetingAsyncTaskData& TTargetingDataStore<FTargetingAsyncTaskData>::FindOrAdd(FTargetingRequestHandle Handle)
 		{
 			return GTargetingAsyncTaskDataStore.Items.FindOrAdd(Handle);
 		}
 
 		template<>
-		FORCEINLINE FTargetingAsyncTaskData* TTargetingDataStore<FTargetingAsyncTaskData>::Find(const FTargetingRequestHandle& Handle)
+		FORCEINLINE FTargetingAsyncTaskData* TTargetingDataStore<FTargetingAsyncTaskData>::Find(FTargetingRequestHandle Handle)
 		{
 			return GTargetingAsyncTaskDataStore.Items.Find(Handle);
 		}
 
 		template<>
-		FORCEINLINE void TTargetingDataStore<FTargetingAsyncTaskData>::OnTargetingRequestHandleReleased(const FTargetingRequestHandle& Handle)
+		FORCEINLINE void TTargetingDataStore<FTargetingAsyncTaskData>::OnTargetingRequestHandleReleased(FTargetingRequestHandle Handle)
 		{
 			GTargetingAsyncTaskDataStore.Items.Remove(Handle);
 		}
@@ -171,19 +171,19 @@ namespace UE
 		extern TTargetingDataStore<FTargetingImmediateTaskData> GTargetingImmediateTaskDataStore;
 
 		template<>
-		FORCEINLINE FTargetingImmediateTaskData& TTargetingDataStore<FTargetingImmediateTaskData>::FindOrAdd(const FTargetingRequestHandle& Handle)
+		FORCEINLINE FTargetingImmediateTaskData& TTargetingDataStore<FTargetingImmediateTaskData>::FindOrAdd(FTargetingRequestHandle Handle)
 		{
 			return GTargetingImmediateTaskDataStore.Items.FindOrAdd(Handle);
 		}
 
 		template<>
-		FORCEINLINE FTargetingImmediateTaskData* TTargetingDataStore<FTargetingImmediateTaskData>::Find(const FTargetingRequestHandle& Handle)
+		FORCEINLINE FTargetingImmediateTaskData* TTargetingDataStore<FTargetingImmediateTaskData>::Find(FTargetingRequestHandle Handle)
 		{
 			return GTargetingImmediateTaskDataStore.Items.Find(Handle);
 		}
 
 		template<>
-		FORCEINLINE void TTargetingDataStore<FTargetingImmediateTaskData>::OnTargetingRequestHandleReleased(const FTargetingRequestHandle& Handle)
+		FORCEINLINE void TTargetingDataStore<FTargetingImmediateTaskData>::OnTargetingRequestHandleReleased(FTargetingRequestHandle Handle)
 		{
 			GTargetingImmediateTaskDataStore.Items.Remove(Handle);
 		}
@@ -194,19 +194,19 @@ namespace UE
 		extern TTargetingDataStore<FTargetingDebugData> GTargetingDebugDataStore;
 
 		template<>
-		FORCEINLINE FTargetingDebugData& TTargetingDataStore<FTargetingDebugData>::FindOrAdd(const FTargetingRequestHandle& Handle)
+		FORCEINLINE FTargetingDebugData& TTargetingDataStore<FTargetingDebugData>::FindOrAdd(FTargetingRequestHandle Handle)
 		{
 			return GTargetingDebugDataStore.Items.FindOrAdd(Handle);
 		}
 
 		template<>
-		FORCEINLINE FTargetingDebugData* TTargetingDataStore<FTargetingDebugData>::Find(const FTargetingRequestHandle& Handle)
+		FORCEINLINE FTargetingDebugData* TTargetingDataStore<FTargetingDebugData>::Find(FTargetingRequestHandle Handle)
 		{
 			return GTargetingDebugDataStore.Items.Find(Handle);
 		}
 
 		template<>
-		FORCEINLINE void TTargetingDataStore<FTargetingDebugData>::OnTargetingRequestHandleReleased(const FTargetingRequestHandle& Handle)
+		FORCEINLINE void TTargetingDataStore<FTargetingDebugData>::OnTargetingRequestHandleReleased(FTargetingRequestHandle Handle)
 		{
 			GTargetingDebugDataStore.Items.Remove(Handle);
 		}
@@ -244,19 +244,19 @@ namespace UE
 		extern TTargetingDataStore<DataType> GTargetingDataStore##DataType;\
 \
 		template<>\
-		FORCEINLINE DataType& TTargetingDataStore<DataType>::FindOrAdd(const FTargetingRequestHandle& Handle)\
+		FORCEINLINE DataType& TTargetingDataStore<DataType>::FindOrAdd(FTargetingRequestHandle Handle)\
 		{\
 			return GTargetingDataStore##DataType.Items.FindOrAdd(Handle);\
 		}\
 \
 		template<>\
-		FORCEINLINE DataType* TTargetingDataStore<DataType>::Find(const FTargetingRequestHandle& Handle)\
+		FORCEINLINE DataType* TTargetingDataStore<DataType>::Find(FTargetingRequestHandle Handle)\
 		{\
 			return GTargetingDataStore##DataType.Items.Find(Handle);\
 		}\
 \
 		template<>\
-		FORCEINLINE void TTargetingDataStore<DataType>::OnTargetingRequestHandleReleased(const FTargetingRequestHandle& Handle)\
+		FORCEINLINE void TTargetingDataStore<DataType>::OnTargetingRequestHandleReleased(FTargetingRequestHandle Handle)\
 		{\
 			GTargetingDataStore##DataType.Items.Remove(Handle);\
 		}\
