Move TargetingSubsystem UFunctions into the same metadata Category

diff --git a/Engine/Plugins/Experimental/GameplayTargetingSystem/Source/GameplayTargetingSystem/TargetingSystem/TargetingSubsystem.h b/Engine/Plugins/Experimental/GameplayTargetingSystem/Source/GameplayTargetingSystem/TargetingSystem/TargetingSubsystem.h
--- a/Engine/Plugins/Experimental/GameplayTargetingSystem/Source/GameplayTargetingSystem/TargetingSystem/TargetingSubsystem.h
+++ b/Engine/Plugins/Experimental/GameplayTargetingSystem/Source/GameplayTargetingSystem/TargetingSystem/TargetingSubsystem.h
@@ -134,18 +134,18 @@ public:
 	TARGETINGSYSTEM_API void ExecuteTargetingRequestWithHandle(FTargetingRequestHandle TargetingHandle, FTargetingRequestDelegate CompletionDelegate) const;
 
 	/** Method to execute an immediate targeting request based on a gameplay targeting preset.*/
-	UFUNCTION(BlueprintCallable, BlueprintPure = false, Category = "Gameplay Targeting | Instant Request")
+	UFUNCTION(BlueprintCallable, BlueprintPure = false, Category = "Targeting System | Instant Request")
 	void ExecuteTargetingRequest(const UTargetingPreset* TargetingPreset, const FTargetingSourceContext& InSourceContext, FTargetingRequestDynamicDelegate CompletionDynamicDelegate) const;
 
 	/** Method to queue an async targeting request with a given targeting handle. */
 	TARGETINGSYSTEM_API void StartAsyncTargetingRequestWithHandle(FTargetingRequestHandle TargetingHandle, FTargetingRequestDelegate CompletionDelegate);
 
 	/** Method to remove an async targeting request with a given targeting handle */
-	UFUNCTION(BlueprintCallable, Category = "Gameplay Targeting | Async Request")
+	UFUNCTION(BlueprintCallable, Category = "Targeting System | Async Request")
 	TARGETINGSYSTEM_API void RemoveAsyncTargetingRequestWithHandle(FTargetingRequestHandle& TargetingHandle);
 
 	/** Method to queue an async targeting request based on a gameplay targeting preset. */
-	UFUNCTION(BlueprintCallable, Category = "Gameplay Targeting | Async Request")
+	UFUNCTION(BlueprintCallable, Category = "Targeting System | Async Request")
 	FTargetingRequestHandle StartAsyncTargetingRequest(const UTargetingPreset* TargetingPreset, const FTargetingSourceContext& InSourceContext, FTargetingRequestDynamicDelegate CompletionDynamicDelegate);
 
 private:
