[Targeting System] Return a copy of the targeting handle instead of a reference.

diff --git a/Engine/Plugins/Experimental/GameplayTargetingSystem/Source/GameplayTargetingSystem/Async/AsyncAction_PerformTargeting.h b/Engine/Plugins/Experimental/GameplayTargetingSystem/Source/GameplayTargetingSystem/Async/AsyncAction_PerformTargeting.h
--- a/Engine/Plugins/Experimental/GameplayTargetingSystem/Source/GameplayTargetingSystem/Async/AsyncAction_PerformTargeting.h
+++ b/Engine/Plugins/Experimental/GameplayTargetingSystem/Source/GameplayTargetingSystem/Async/AsyncAction_PerformTargeting.h
@@ -31,7 +31,7 @@ public:
 	static UAsyncAction_PerformTargeting* PerformFilteringRequest(AActor* SourceActor, UTargetingPreset* TargetingPreset, bool bUseAsyncTargeting, const TArray<AActor*> InTargets);
 
 	UFUNCTION(BlueprintPure, Category = "Targeting")
-	const FTargetingRequestHandle& GetTargetingHandle() const { return TargetingHandle; }
+	FTargetingRequestHandle GetTargetingHandle() const { return TargetingHandle; }
 
 	virtual void Activate() override;
 
