[Targeting System] Adding more logging for handle release and Task processing

#[review]-25165753 [REDACTED_HANDLE], [REDACTED_HANDLE], [REDACTED_HANDLE], [REDACTED_HANDLE], [REDACTED_HANDLE], [REDACTED_HANDLE], [REDACTED_HANDLE], [REDACTED_HANDLE]

diff --git a/Engine/Plugins/Experimental/GameplayTargetingSystem/Source/GameplayTargetingSystem/TargetingSystem/TargetingSubsystem.cpp b/Engine/Plugins/Experimental/GameplayTargetingSystem/Source/GameplayTargetingSystem/TargetingSystem/TargetingSubsystem.cpp
--- a/Engine/Plugins/Experimental/GameplayTargetingSystem/Source/GameplayTargetingSystem/TargetingSystem/TargetingSubsystem.cpp
+++ b/Engine/Plugins/Experimental/GameplayTargetingSystem/Source/GameplayTargetingSystem/TargetingSystem/TargetingSubsystem.cpp
@@ -188,6 +188,7 @@ void UTargetingSubsystem::Tick(float DeltaTime)
 	const int32 NumRequests = AsyncTargetingRequests.Num();
 	{
 		bTickingAsycnRequests = true;
+		TARGETING_LOG(Verbose, TEXT("UTargetingSubsystem::Tick - Starting to Process %d requests"), NumRequests);
 
 		float TimeLeft = TargetingSystemCVars::MaxAsyncTickTime;
 		for (int32 RequestIterator = 0; RequestIterator < NumRequests; ++RequestIterator)
@@ -321,6 +322,7 @@ void UTargetingSubsystem::ReleaseTargetRequestHandle(FTargetingRequestHandle& Ha
 #endif // ENABLE_DRAW_DEBUG
 
 	ReleaseHandleDelegate.Broadcast(CachedHandle);
+	TARGETING_LOG(Verbose, TEXT("%s: - Releasigng Handle [%d]"), ANSI_TO_TCHAR(__FUNCTION__), CachedHandle.Handle);
 }
 
 void UTargetingSubsystem::ExecuteTargetingRequestWithHandle(FTargetingRequestHandle TargetingHandle, FTargetingRequestDelegate CompletionDelegate) const
