[Targeting System] Fixing an issue where more entries would be added to the async request array while processing them. Using the same mechanism that prevents unsafe removals to queue adds while we're processsing them. Once that lock is released, we add all the pending requests that came in while processing.

diff --git a/Engine/Plugins/Experimental/GameplayTargetingSystem/Source/GameplayTargetingSystem/TargetingSystem/TargetingSubsystem.cpp b/Engine/Plugins/Experimental/GameplayTargetingSystem/Source/GameplayTargetingSystem/TargetingSystem/TargetingSubsystem.cpp
--- a/Engine/Plugins/Experimental/GameplayTargetingSystem/Source/GameplayTargetingSystem/TargetingSystem/TargetingSubsystem.cpp
+++ b/Engine/Plugins/Experimental/GameplayTargetingSystem/Source/GameplayTargetingSystem/TargetingSystem/TargetingSubsystem.cpp
@@ -215,6 +215,7 @@ void UTargetingSubsystem::Tick(float DeltaTime)
 		}
 
 		bTickingAsycnRequests = false;
+		OnFinishedTickingAsyncRequests();
 	}
 	
 	{
@@ -269,6 +270,22 @@ void UTargetingSubsystem::Tick(float DeltaTime)
 	}
 }
 
+void UTargetingSubsystem::OnFinishedTickingAsyncRequests()
+{
+	// Queue all the pending requests once we're safe to add to the Targeting Requests Data Store
+	for (const TPair<FTargetingRequestHandle, FTargetingRequestData>& PendingRequest : PendingTargetingRequests)
+	{
+		ExecuteTargetingRequestWithHandle(PendingRequest.Key, PendingRequest.Value.TargetingRequestDelegate);
+	}
+	PendingTargetingRequests.Reset();
+
+	for (const TPair<FTargetingRequestHandle, FTargetingRequestData>& PendingAsyncRequest : PendingAsyncTargetingRequests)
+	{
+		StartAsyncTargetingRequestWithHandle(PendingAsyncRequest.Key, PendingAsyncRequest.Value.TargetingRequestDelegate, PendingAsyncRequest.Value.TargetingRequestDynamicDelegate);
+	}
+	PendingAsyncTargetingRequests.Reset();
+}
+
 bool UTargetingSubsystem::IsTickable() const
 {
 	return !HasAnyFlags(RF_BeginDestroyed) && IsValidChecked(this) && (AsyncTargetingRequests.Num() > 0);
@@ -325,9 +342,16 @@ void UTargetingSubsystem::ReleaseTargetRequestHandle(FTargetingRequestHandle& Ha
 	TARGETING_LOG(Verbose, TEXT("%s: - Releasigng Handle [%d]"), ANSI_TO_TCHAR(__FUNCTION__), CachedHandle.Handle);
 }
 
-void UTargetingSubsystem::ExecuteTargetingRequestWithHandle(FTargetingRequestHandle TargetingHandle, FTargetingRequestDelegate CompletionDelegate) const
+void UTargetingSubsystem::ExecuteTargetingRequestWithHandle(FTargetingRequestHandle TargetingHandle, FTargetingRequestDelegate CompletionDelegate, FTargetingRequestDynamicDelegate CompletionDynamicDelegate)
 {
-	FTargetingRequestDynamicDelegate CompletionDynamicDelegate;
+	// If we're processing Targeting Requests, queue this one to prevent memory stomps. We'll add it once we finish Ticking Async Requests.
+	if (bTickingAsycnRequests)
+	{
+		FTargetingRequestData& RequestData = PendingTargetingRequests.FindOrAdd(TargetingHandle);
+		RequestData.Initialize(CompletionDelegate, CompletionDynamicDelegate);
+		return;
+	}
+
 	ExecuteTargetingRequestWithHandleInternal(TargetingHandle, CompletionDelegate, CompletionDynamicDelegate);
 
 	// if the caller setup task data to release on completion, lets do it
@@ -340,7 +364,7 @@ void UTargetingSubsystem::ExecuteTargetingRequestWithHandle(FTargetingRequestHan
 	}
 }
 
-void UTargetingSubsystem::ExecuteTargetingRequest(const UTargetingPreset* TargetingPreset, const FTargetingSourceContext& SourceContext, FTargetingRequestDynamicDelegate CompletionDynamicDelegate) const
+void UTargetingSubsystem::ExecuteTargetingRequest(const UTargetingPreset* TargetingPreset, const FTargetingSourceContext& SourceContext, FTargetingRequestDynamicDelegate CompletionDynamicDelegate)
 {
 	FTargetingRequestHandle RequestHandle;
 	if (ensure(TargetingPreset))
@@ -357,15 +381,30 @@ void UTargetingSubsystem::ExecuteTargetingRequest(const UTargetingPreset* Target
 #endif // WITH_EDITORONLY_DATA
 #endif // ENABLE_DRAW_DEBUG
 
+		// If we're processing Targeting Requests, queue this one to prevent memory stomps. We'll add it once we finish Ticking Async Requests.
 		FTargetingRequestDelegate CompletionDelegate;
+		if (bTickingAsycnRequests)
+		{
+			FTargetingRequestData& RequestData = PendingTargetingRequests.FindOrAdd(RequestHandle);
+			RequestData.Initialize(CompletionDelegate, CompletionDynamicDelegate);
+			return;
+		}
+
 		ExecuteTargetingRequestWithHandleInternal(RequestHandle, CompletionDelegate, CompletionDynamicDelegate);
 		UTargetingSubsystem::ReleaseTargetRequestHandle(RequestHandle);
 	}
 }
 
-void UTargetingSubsystem::StartAsyncTargetingRequestWithHandle(FTargetingRequestHandle TargetingHandle, FTargetingRequestDelegate CompletionDelegate)
+void UTargetingSubsystem::StartAsyncTargetingRequestWithHandle(FTargetingRequestHandle TargetingHandle, FTargetingRequestDelegate CompletionDelegate, FTargetingRequestDynamicDelegate CompletionDynamicDelegate)
 {
-	FTargetingRequestDynamicDelegate CompletionDynamicDelegate;
+	// If we're processing Targeting Requests, queue this one to prevent memory stomps. We'll add it once we finish Ticking Async Requests.
+	if (bTickingAsycnRequests)
+	{
+		FTargetingRequestData& RequestData = PendingAsyncTargetingRequests.FindOrAdd(TargetingHandle);
+		RequestData.Initialize(CompletionDelegate, CompletionDynamicDelegate);
+		return;
+	}
+
 	StartAsyncTargetingRequestWithHandleInternal(TargetingHandle, CompletionDelegate, CompletionDynamicDelegate);
 }
 
@@ -435,6 +474,15 @@ FTargetingRequestHandle UTargetingSubsystem::StartAsyncTargetingRequest(const UT
 #endif // ENABLE_DRAW_DEBUG
 
 		FTargetingRequestDelegate CompletionDelegate;
+
+		// If we're processing Targeting Requests, queue this one to prevent memory stomps. We'll add it once we finish Ticking Async Requests.
+		if (bTickingAsycnRequests)
+		{
+			FTargetingRequestData& RequestData = PendingAsyncTargetingRequests.FindOrAdd(RequestHandle);
+			RequestData.Initialize(CompletionDelegate, CompletionDynamicDelegate);
+			return RequestHandle;
+		}
+
 		StartAsyncTargetingRequestWithHandleInternal(RequestHandle, CompletionDelegate, CompletionDynamicDelegate);
 	}
 	
diff --git a/Engine/Plugins/Experimental/GameplayTargetingSystem/Source/GameplayTargetingSystem/TargetingSystem/TargetingSubsystem.h b/Engine/Plugins/Experimental/GameplayTargetingSystem/Source/GameplayTargetingSystem/TargetingSystem/TargetingSubsystem.h
--- a/Engine/Plugins/Experimental/GameplayTargetingSystem/Source/GameplayTargetingSystem/TargetingSystem/TargetingSubsystem.h
+++ b/Engine/Plugins/Experimental/GameplayTargetingSystem/Source/GameplayTargetingSystem/TargetingSystem/TargetingSubsystem.h
@@ -1,13 +1,14 @@
 // Copyright Epic Games, Inc. All Rights Reserved.
 #pragma once
 
+#include "Containers/SortedMap.h"
 #include "CoreMinimal.h"
 #include "DrawDebugHelpers.h"
-#include "Tasks/TargetingTask.h"
-#include "Types/TargetingSystemTypes.h"
 #include "Misc/CoreMisc.h"
 #include "Subsystems/GameInstanceSubsystem.h"
+#include "Tasks/TargetingTask.h"
 #include "Tickable.h"
+#include "Types/TargetingSystemTypes.h"
 
 #include "TargetingSubsystem.generated.h"
 
@@ -131,14 +132,14 @@ public:
 	/** Targeting Request Methods */
 public:
 	/** Method to execute an immediate targeting request with a given targeting handle. */
-	TARGETINGSYSTEM_API void ExecuteTargetingRequestWithHandle(FTargetingRequestHandle TargetingHandle, FTargetingRequestDelegate CompletionDelegate) const;
+	TARGETINGSYSTEM_API void ExecuteTargetingRequestWithHandle(FTargetingRequestHandle TargetingHandle, FTargetingRequestDelegate CompletionDelegate = FTargetingRequestDelegate(), FTargetingRequestDynamicDelegate CompletionDynamicDelegate = FTargetingRequestDynamicDelegate());
 
 	/** Method to execute an immediate targeting request based on a gameplay targeting preset.*/
 	UFUNCTION(BlueprintCallable, BlueprintPure = false, Category = "Targeting System | Instant Request")
-	void ExecuteTargetingRequest(const UTargetingPreset* TargetingPreset, const FTargetingSourceContext& InSourceContext, FTargetingRequestDynamicDelegate CompletionDynamicDelegate) const;
+	void ExecuteTargetingRequest(const UTargetingPreset* TargetingPreset, const FTargetingSourceContext& InSourceContext, FTargetingRequestDynamicDelegate CompletionDynamicDelegate);
 
 	/** Method to queue an async targeting request with a given targeting handle. */
-	TARGETINGSYSTEM_API void StartAsyncTargetingRequestWithHandle(FTargetingRequestHandle TargetingHandle, FTargetingRequestDelegate CompletionDelegate);
+	TARGETINGSYSTEM_API void StartAsyncTargetingRequestWithHandle(FTargetingRequestHandle TargetingHandle, FTargetingRequestDelegate CompletionDelegate = FTargetingRequestDelegate(), FTargetingRequestDynamicDelegate CompletionDynamicDelegate = FTargetingRequestDynamicDelegate());
 
 	/** Method to remove an async targeting request with a given targeting handle */
 	UFUNCTION(BlueprintCallable, Category = "Targeting System | Async Request")
@@ -168,6 +169,12 @@ private:
 	/** Flag indicating the targeting system is currently in its tick processing targeting tasks for async request */
 	bool bTickingAsycnRequests = false;
 
+	/** While we're processing targeting requests, we add any incoming requests to this array to prevent memory stomps */
+	TSortedMap<FTargetingRequestHandle, FTargetingRequestData> PendingTargetingRequests;
+
+	/** (Version for Async Requests) While we're processing targeting requests, we add any incoming requests to this array to prevent memory stomps */
+	TSortedMap<FTargetingRequestHandle, FTargetingRequestData> PendingAsyncTargetingRequests;
+
 	/** ~Targeting Request Methods */
 
 	/** Blueprint Helper Methods */
@@ -212,6 +219,9 @@ private:
 	void AddDebugTrackedImmediateTargetRequests(FTargetingRequestHandle TargetingHandle) const;
 	void AddDebugTrackedAsyncTargetRequests(FTargetingRequestHandle TargetingHandle) const;
 
+	/** Called when we set bTickingAsyncRequests to false, at this point it's safe to perform any queued operations on the Async Requests Array */
+	void OnFinishedTickingAsyncRequests();
+
 	mutable TArray<FTargetingRequestHandle> DebugTrackedImmediateTargetRequests;
 	mutable int32 CurrentImmediateRequestIndex = 0;
 
