[Targeting System] Adding the source location of the AOE trace as the trace start to the hit result so users have that info available.

diff --git a/Engine/Plugins/Experimental/GameplayTargetingSystem/Source/GameplayTargetingSystem/Tasks/TargetingSelectionTask_AOE.cpp b/Engine/Plugins/Experimental/GameplayTargetingSystem/Source/GameplayTargetingSystem/Tasks/TargetingSelectionTask_AOE.cpp
--- a/Engine/Plugins/Experimental/GameplayTargetingSystem/Source/GameplayTargetingSystem/Tasks/TargetingSelectionTask_AOE.cpp
+++ b/Engine/Plugins/Experimental/GameplayTargetingSystem/Source/GameplayTargetingSystem/Tasks/TargetingSelectionTask_AOE.cpp
@@ -231,8 +231,10 @@ void UTargetingSelectionTask_AOE::ProcessOverlapResults(const FTargetingRequestH
 	if (Overlaps.Num() > 0)
 	{
 		FTargetingDefaultResultsSet& TargetingResults = FTargetingDefaultResultsSet::FindOrAdd(TargetingHandle);
+		const FVector SourceLocation = GetSourceLocation(TargetingHandle) + GetSourceOffset(TargetingHandle);
 		for (const FOverlapResult& OverlapResult : Overlaps)
 		{
+
 			if (!OverlapResult.GetActor())
 			{
 				continue;
@@ -242,7 +244,6 @@ void UTargetingSelectionTask_AOE::ProcessOverlapResults(const FTargetingRequestH
 			if (ShapeType == ETargetingAOEShape::Cylinder)
 			{
 				const float RadiusSquared = (HalfExtent.X * HalfExtent.X);
-				const FVector SourceLocation = GetSourceLocation(TargetingHandle) + GetSourceOffset(TargetingHandle);
 				const float DistanceSquared = FVector::DistSquared2D(OverlapResult.GetActor()->GetActorLocation(), SourceLocation);
 				if (DistanceSquared > RadiusSquared)
 				{
@@ -266,6 +267,7 @@ void UTargetingSelectionTask_AOE::ProcessOverlapResults(const FTargetingRequestH
 				ResultData->HitResult.HitObjectHandle = FActorInstanceHandle(OverlapResult.GetActor());
 				ResultData->HitResult.Component = OverlapResult.GetComponent();
 				ResultData->HitResult.bBlockingHit = OverlapResult.bBlockingHit;
+				ResultData->HitResult.TraceStart = SourceLocation;
 			}
 		}
 
