PR #10598: Expose UTargetingSubsystem::ReleaseHandleDelegate via API markup

diff --git a/Engine/Plugins/Experimental/GameplayTargetingSystem/Source/GameplayTargetingSystem/TargetingSystem/TargetingSubsystem.h b/Engine/Plugins/Experimental/GameplayTargetingSystem/Source/GameplayTargetingSystem/TargetingSystem/TargetingSubsystem.h
--- a/Engine/Plugins/Experimental/GameplayTargetingSystem/Source/GameplayTargetingSystem/TargetingSystem/TargetingSubsystem.h
+++ b/Engine/Plugins/Experimental/GameplayTargetingSystem/Source/GameplayTargetingSystem/TargetingSystem/TargetingSubsystem.h
@@ -125,7 +125,7 @@ public:
 
 	/** The handle released delegate that fires right before a handle is release so all data stores can clean up their state */
 	DECLARE_MULTICAST_DELEGATE_OneParam(FOnTargetingRequestHandleReleased, FTargetingRequestHandle)
-	static FOnTargetingRequestHandleReleased ReleaseHandleDelegate;
+	TARGETINGSYSTEM_API static FOnTargetingRequestHandleReleased ReleaseHandleDelegate;
 
 	/** Target Handle Generation Methods */
 
