PR #10536: TargetingSystem: Export UAbilityTask_PerformTargeting to other C++ modules

diff --git a/Engine/Plugins/Experimental/GameplayTargetingSystem/Source/GameplayTargetingSystem/AbilityTasks/AbilityTask_PerformTargeting.h b/Engine/Plugins/Experimental/GameplayTargetingSystem/Source/GameplayTargetingSystem/AbilityTasks/AbilityTask_PerformTargeting.h
--- a/Engine/Plugins/Experimental/GameplayTargetingSystem/Source/GameplayTargetingSystem/AbilityTasks/AbilityTask_PerformTargeting.h
+++ b/Engine/Plugins/Experimental/GameplayTargetingSystem/Source/GameplayTargetingSystem/AbilityTasks/AbilityTask_PerformTargeting.h
@@ -13,10 +13,11 @@ class UTargetingPreset;
 DECLARE_DYNAMIC_MULTICAST_DELEGATE_OneParam(FTargetReadyDelegate, FTargetingRequestHandle, TargetingRequestHandle);
 
 UCLASS()
-class UAbilityTask_PerformTargeting : public UAbilityTask
+class TARGETINGSYSTEM_API UAbilityTask_PerformTargeting : public UAbilityTask
 {
 	GENERATED_BODY()
 
+public:
 	// Called when the targeting request has been completed and results are ready
 	UPROPERTY(BlueprintAssignable)
 	FTargetReadyDelegate OnTargetReady;
