Fix comment for bReleaseOnCompletion in TargetingSystemTypes.h

#rnx

diff --git a/Engine/Plugins/Experimental/GameplayTargetingSystem/Source/GameplayTargetingSystem/Public/Types/TargetingSystemTypes.h b/Engine/Plugins/Experimental/GameplayTargetingSystem/Source/GameplayTargetingSystem/Public/Types/TargetingSystemTypes.h
--- a/Engine/Plugins/Experimental/GameplayTargetingSystem/Source/GameplayTargetingSystem/Public/Types/TargetingSystemTypes.h
+++ b/Engine/Plugins/Experimental/GameplayTargetingSystem/Source/GameplayTargetingSystem/Public/Types/TargetingSystemTypes.h
@@ -258,7 +258,7 @@ public:
 	/** Indicates the targeting request was an asynchronous request */
 	uint8 bAsyncRequest : 1;
 
-	/** Indicates the targeting request was an asynchronous request */
+	/** Indicates the targeting request should be released after completion */
 	uint8 bReleaseOnCompletion : 1;
 
 	/** Indicates the target request should be re-queued after completion */
