Targeting on Client/Debug Visualization
#rnx

diff --git a/Engine/Plugins/Experimental/GameplayTargetingSystem/Source/GameplayTargetingSystem/Private/TargetingSystem/TargetingSubsystem.cpp b/Engine/Plugins/Experimental/GameplayTargetingSystem/Source/GameplayTargetingSystem/Private/TargetingSystem/TargetingSubsystem.cpp
--- a/Engine/Plugins/Experimental/GameplayTargetingSystem/Source/GameplayTargetingSystem/Private/TargetingSystem/TargetingSubsystem.cpp
+++ b/Engine/Plugins/Experimental/GameplayTargetingSystem/Source/GameplayTargetingSystem/Private/TargetingSystem/TargetingSubsystem.cpp
@@ -64,6 +64,13 @@ namespace TargetingSystemCVars
 		TotalDebugRecentRequestsTracked,
 		TEXT("Sets the total # of targeting requests that will be tracked upon starting of them (default = 5)")
 	);
+	
+	static float OverrideTargetingLifeTime = 0.f;
+    FAutoConsoleVariableRef CvarOverrideTargetingLifeTime(
+    	TEXT("ts.debug.OverrideTargetingLifeTime"),
+    	OverrideTargetingLifeTime,
+    	TEXT("Overrides the draws life time to ease the debugging")
+    );
 
 #endif // ENABLE_DRAW_DEBUG
 }
@@ -749,6 +756,11 @@ bool UTargetingSubsystem::IsTargetingDebugEnabled()
 	return TargetingSystemCVars::bEnableTargetingDebugging;
 }
 
+float UTargetingSubsystem::GetOverrideTargetingLifeTime()
+{
+	return TargetingSystemCVars::OverrideTargetingLifeTime;
+}
+
 /* static */
 void UTargetingSubsystem::OnShowDebugInfo(AHUD* HUD, UCanvas* Canvas, const FDebugDisplayInfo& DisplayInfo, float& YL, float& YPos)
 {
diff --git a/Engine/Plugins/Experimental/GameplayTargetingSystem/Source/GameplayTargetingSystem/Private/Tasks/TargetingSelectionTask_AOE.cpp b/Engine/Plugins/Experimental/GameplayTargetingSystem/Source/GameplayTargetingSystem/Private/Tasks/TargetingSelectionTask_AOE.cpp
--- a/Engine/Plugins/Experimental/GameplayTargetingSystem/Source/GameplayTargetingSystem/Private/Tasks/TargetingSelectionTask_AOE.cpp
+++ b/Engine/Plugins/Experimental/GameplayTargetingSystem/Source/GameplayTargetingSystem/Private/Tasks/TargetingSelectionTask_AOE.cpp
@@ -118,7 +118,7 @@ void UTargetingSelectionTask_AOE::DebugDrawBoundingVolume(const FTargetingReques
 	const FCollisionShape CollisionShape = GetCollisionShape();
 
 	const bool bPersistentLines = false;
-	const float LifeTime = 0.0f;
+	const float LifeTime = UTargetingSubsystem::GetOverrideTargetingLifeTime();
 	const uint8 DepthPriority = 0;
 	const float Thickness = 2.0f;
 
diff --git a/Engine/Plugins/Experimental/GameplayTargetingSystem/Source/GameplayTargetingSystem/Public/TargetingSystem/TargetingSubsystem.h b/Engine/Plugins/Experimental/GameplayTargetingSystem/Source/GameplayTargetingSystem/Public/TargetingSystem/TargetingSubsystem.h
--- a/Engine/Plugins/Experimental/GameplayTargetingSystem/Source/GameplayTargetingSystem/Public/TargetingSystem/TargetingSubsystem.h
+++ b/Engine/Plugins/Experimental/GameplayTargetingSystem/Source/GameplayTargetingSystem/Public/TargetingSystem/TargetingSubsystem.h
@@ -207,6 +207,9 @@ public:
 	/** Helper method to indicate if the debugging feature is enabled */
 	TARGETINGSYSTEM_API static bool IsTargetingDebugEnabled();
 
+	/** Helper method to overrides the life time of the shapes drawn */
+	static float GetOverrideTargetingLifeTime();
+
 	/** Callback to handle OnShowDebugInfo to check if the TargetingSystem has been enabled */
 	static void OnShowDebugInfo(AHUD* HUD, UCanvas* Canvas, const FDebugDisplayInfo& DisplayInfo, float& YL, float& YPos);
 
