Add Blueprintable TargetingTasks for GameplayTargetingSystem

diff --git a/Engine/Plugins/Experimental/GameplayTargetingSystem/Source/GameplayTargetingSystem/Private/Tasks/SimpleTargetingFilterTask.cpp b/Engine/Plugins/Experimental/GameplayTargetingSystem/Source/GameplayTargetingSystem/Private/Tasks/SimpleTargetingFilterTask.cpp
new file mode 100644
--- /dev/null
+++ b/Engine/Plugins/Experimental/GameplayTargetingSystem/Source/GameplayTargetingSystem/Private/Tasks/SimpleTargetingFilterTask.cpp
@@ -0,0 +1,8 @@
+﻿// Copyright Epic Games, Inc. All Rights Reserved.
+
+#include "Tasks/SimpleTargetingFilterTask.h"
+
+bool USimpleTargetingFilterTask::ShouldFilterTarget(const FTargetingRequestHandle& TargetingHandle, const FTargetingDefaultResultData& TargetData) const
+{
+	return BP_ShouldFilterTarget(TargetingHandle, TargetData);
+}
diff --git a/Engine/Plugins/Experimental/GameplayTargetingSystem/Source/GameplayTargetingSystem/Private/Tasks/SimpleTargetingSelectionTask.cpp b/Engine/Plugins/Experimental/GameplayTargetingSystem/Source/GameplayTargetingSystem/Private/Tasks/SimpleTargetingSelectionTask.cpp
new file mode 100644
--- /dev/null
+++ b/Engine/Plugins/Experimental/GameplayTargetingSystem/Source/GameplayTargetingSystem/Private/Tasks/SimpleTargetingSelectionTask.cpp
@@ -0,0 +1,45 @@
+﻿// Copyright Epic Games, Inc. All Rights Reserved.
+
+#include "Tasks/SimpleTargetingSelectionTask.h"
+
+void USimpleTargetingSelectionTask::Execute(const FTargetingRequestHandle& TargetingHandle) const
+{
+	Super::Execute(TargetingHandle);
+
+	SetTaskAsyncState(TargetingHandle, ETargetingTaskAsyncState::Executing);
+
+	if (FTargetingSourceContext* SourceContext = FTargetingSourceContext::Find(TargetingHandle))
+	{
+		// Call the blueprint function
+		SelectTargets(TargetingHandle, *SourceContext);
+	}
+	SetTaskAsyncState(TargetingHandle, ETargetingTaskAsyncState::Completed);
+}
+
+bool USimpleTargetingSelectionTask::AddTargetActor(const FTargetingRequestHandle& TargetingHandle, AActor* Actor) const
+{
+	if (Actor)
+	{
+		FTargetingDefaultResultsSet& ResultsSet = FTargetingDefaultResultsSet::FindOrAdd(TargetingHandle);
+		if (!ResultsSet.TargetResults.FindByPredicate([Actor](const FTargetingDefaultResultData& ResultData){ return ResultData.HitResult.GetActor() == Actor; }))
+		{
+			FTargetingDefaultResultData& ResultData = ResultsSet.TargetResults.AddDefaulted_GetRef();
+			ResultData.HitResult.HitObjectHandle = FActorInstanceHandle(Actor);
+			ResultData.HitResult.Location = Actor->GetActorLocation();
+			return true;
+		}
+	}
+	return false;
+}
+
+bool USimpleTargetingSelectionTask::AddHitResult(const FTargetingRequestHandle& TargetingHandle, const FHitResult& HitResult) const
+{
+	FTargetingDefaultResultsSet& ResultsSet = FTargetingDefaultResultsSet::FindOrAdd(TargetingHandle);
+	if (!ResultsSet.TargetResults.FindByPredicate([HitResult](const FTargetingDefaultResultData& ResultData){ return ResultData.HitResult.GetActor() == HitResult.GetActor(); }))
+	{
+		FTargetingDefaultResultData& ResultData = ResultsSet.TargetResults.AddDefaulted_GetRef();
+		ResultData.HitResult = HitResult;
+		return true;
+	}
+	return false;
+}
diff --git a/Engine/Plugins/Experimental/GameplayTargetingSystem/Source/GameplayTargetingSystem/Private/Tasks/SimpleTargetingSortTask.cpp b/Engine/Plugins/Experimental/GameplayTargetingSystem/Source/GameplayTargetingSystem/Private/Tasks/SimpleTargetingSortTask.cpp
new file mode 100644
--- /dev/null
+++ b/Engine/Plugins/Experimental/GameplayTargetingSystem/Source/GameplayTargetingSystem/Private/Tasks/SimpleTargetingSortTask.cpp
@@ -0,0 +1,8 @@
+﻿// Copyright Epic Games, Inc. All Rights Reserved.
+
+#include "Tasks/SimpleTargetingSortTask.h"
+
+float USimpleTargetingSortTask::GetScoreForTarget(const FTargetingRequestHandle& TargetingHandle, const FTargetingDefaultResultData& TargetData) const
+{
+	return BP_GetScoreForTarget(TargetingHandle, TargetData);
+}
diff --git a/Engine/Plugins/Experimental/GameplayTargetingSystem/Source/GameplayTargetingSystem/Public/Tasks/SimpleTargetingFilterTask.h b/Engine/Plugins/Experimental/GameplayTargetingSystem/Source/GameplayTargetingSystem/Public/Tasks/SimpleTargetingFilterTask.h
new file mode 100644
--- /dev/null
+++ b/Engine/Plugins/Experimental/GameplayTargetingSystem/Source/GameplayTargetingSystem/Public/Tasks/SimpleTargetingFilterTask.h
@@ -0,0 +1,32 @@
+// Copyright Epic Games, Inc. All Rights Reserved.
+#pragma once
+
+#include "UObject/Object.h"
+#include "Tasks/TargetingFilterTask_BasicFilterTemplate.h"
+
+#include "SimpleTargetingFilterTask.generated.h"
+
+struct FTargetingRequestHandle;
+struct FTargetingDefaultResultData;
+
+/**
+*	@class USimpleTargetingSelectionTask
+*
+*	This is a blueprintable TargetingTask class made for filtering out Targets from the results list.
+*	Override the ShouldRemoveTarget Blueprint Function to define the rules for this filter.
+*/
+UCLASS(EditInlineNew, Abstract, Blueprintable)
+class TARGETINGSYSTEM_API USimpleTargetingFilterTask : public UTargetingFilterTask_BasicFilterTemplate
+{
+	GENERATED_BODY()
+
+public:
+	//~ Begin UTargetingFilterTask_BasicFilterTemplate Interface
+	virtual bool ShouldFilterTarget(const FTargetingRequestHandle& TargetingHandle, const FTargetingDefaultResultData& TargetData) const override;
+	//~ End UTargetingFilterTask_BasicFilterTemplate Interface
+
+	/** Returns true if a Target should be removed from the results of this TargetingRequest */
+	UFUNCTION(BlueprintImplementableEvent, DisplayName=ShouldRemoveTarget)
+	bool BP_ShouldFilterTarget(const FTargetingRequestHandle& TargetingHandle, const FTargetingDefaultResultData& TargetData) const;
+
+};
diff --git a/Engine/Plugins/Experimental/GameplayTargetingSystem/Source/GameplayTargetingSystem/Public/Tasks/SimpleTargetingSelectionTask.h b/Engine/Plugins/Experimental/GameplayTargetingSystem/Source/GameplayTargetingSystem/Public/Tasks/SimpleTargetingSelectionTask.h
new file mode 100644
--- /dev/null
+++ b/Engine/Plugins/Experimental/GameplayTargetingSystem/Source/GameplayTargetingSystem/Public/Tasks/SimpleTargetingSelectionTask.h
@@ -0,0 +1,50 @@
+// Copyright Epic Games, Inc. All Rights Reserved.
+#pragma once
+
+#include "UObject/Object.h"
+#include "TargetingTask.h"
+
+#include "SimpleTargetingSelectionTask.generated.h"
+
+class UTargetingSubsystem;
+struct FTargetingRequestHandle;
+struct FTargetingSourceContext;
+
+/**
+*	@class USimpleTargetingSelectionTask
+*
+*	This is a blueprintable TargetingTask class made for adding new Targets to the results list.
+*	Override the SelectTargets function and call AddTargetHitResult or AddTargetActor to select new targets.
+*/
+UCLASS(EditInlineNew, Abstract, Blueprintable)
+class TARGETINGSYSTEM_API USimpleTargetingSelectionTask : public UTargetingTask
+{
+	GENERATED_BODY()
+
+public:
+	/** Evaluation function called by derived classes to process the targeting request */
+	virtual void Execute(const FTargetingRequestHandle& TargetingHandle) const override;
+	
+	UFUNCTION(BlueprintImplementableEvent)
+	void SelectTargets(const FTargetingRequestHandle& TargetingHandle, const FTargetingSourceContext& SourceContext) const;
+
+protected:
+	/**
+	 * Adds a single Actor to the Targeting Results for a given TargetingRequestHandle.
+	 * Returns false when the Actor was already in the list.
+	 * 
+	 * NOTE: If you have a HitResult associated with this selection, please use AddHitResult instead.
+	 */
+	UFUNCTION(BlueprintCallable, BlueprintPure=False, Category=Targeting)
+	bool AddTargetActor(const FTargetingRequestHandle& TargetingHandle, AActor* Actor) const;
+
+	/**
+	 * Adds a HitResult to the Targeting Results for a given TargetingRequestHandle.
+	 * Returns False when the Actor that was hit was already in the list
+	 */
+	UFUNCTION(BlueprintCallable, BlueprintPure=False, Category=Targeting)
+	bool AddHitResult(const FTargetingRequestHandle& TargetingHandle, const FHitResult& HitResult) const;
+	
+};
+
+
diff --git a/Engine/Plugins/Experimental/GameplayTargetingSystem/Source/GameplayTargetingSystem/Public/Tasks/SimpleTargetingSortTask.h b/Engine/Plugins/Experimental/GameplayTargetingSystem/Source/GameplayTargetingSystem/Public/Tasks/SimpleTargetingSortTask.h
new file mode 100644
--- /dev/null
+++ b/Engine/Plugins/Experimental/GameplayTargetingSystem/Source/GameplayTargetingSystem/Public/Tasks/SimpleTargetingSortTask.h
@@ -0,0 +1,33 @@
+// Copyright Epic Games, Inc. All Rights Reserved.
+#pragma once
+
+#include "UObject/Object.h"
+#include "Tasks/TargetingSortTask_Base.h"
+
+#include "SimpleTargetingSortTask.generated.h"
+
+struct FTargetingRequestHandle;
+
+/**
+*	@class USimpleTargetingSortTask
+*
+*	This is a blueprintable TargetingTask class made for reordering Targets in the results list.
+*	Override the GetScoreForTarget Blueprint Function to provide a score for each element and then Targets will be
+*	sorted by score.
+*/
+UCLASS(EditInlineNew, Abstract, Blueprintable)
+class TARGETINGSYSTEM_API USimpleTargetingSortTask : public UTargetingSortTask_Base
+{
+	GENERATED_BODY()
+
+protected:
+	/** Called on every target to get a Score for sorting. This score will be added to the Score float in FTargetingDefaultResultData */
+	virtual float GetScoreForTarget(const FTargetingRequestHandle& TargetingHandle, const FTargetingDefaultResultData& TargetData) const override;
+
+	/** Called on every target to get a Score for sorting. This score will be added to the Score float in FTargetingDefaultResultData */
+	UFUNCTION(BlueprintImplementableEvent, DisplayName=GetScoreForTarget, Category=Targeting)
+	float BP_GetScoreForTarget(const FTargetingRequestHandle& TargetingHandle, const FTargetingDefaultResultData& TargetData) const;
+
+};
+
+
diff --git a/Engine/Plugins/Experimental/GameplayTargetingSystem/Source/GameplayTargetingSystem/Public/Tasks/TargetingTask.h b/Engine/Plugins/Experimental/GameplayTargetingSystem/Source/GameplayTargetingSystem/Public/Tasks/TargetingTask.h
--- a/Engine/Plugins/Experimental/GameplayTargetingSystem/Source/GameplayTargetingSystem/Public/Tasks/TargetingTask.h
+++ b/Engine/Plugins/Experimental/GameplayTargetingSystem/Source/GameplayTargetingSystem/Public/Tasks/TargetingTask.h
@@ -40,7 +40,7 @@ struct FTargetingDebugInfo;
 *	the end user might prefer to make decisions on. Distance (min/max), score rating
 *	etc.
 */
-UCLASS(EditInlineNew, Abstract)
+UCLASS(EditInlineNew, Abstract, Const)
 class TARGETINGSYSTEM_API UTargetingTask : public UObject
 {
 	GENERATED_BODY()
diff --git a/Engine/Plugins/Experimental/GameplayTargetingSystem/Source/GameplayTargetingSystem/Public/Types/TargetingSystemTypes.h b/Engine/Plugins/Experimental/GameplayTargetingSystem/Source/GameplayTargetingSystem/Public/Types/TargetingSystemTypes.h
--- a/Engine/Plugins/Experimental/GameplayTargetingSystem/Source/GameplayTargetingSystem/Public/Types/TargetingSystemTypes.h
+++ b/Engine/Plugins/Experimental/GameplayTargetingSystem/Source/GameplayTargetingSystem/Public/Types/TargetingSystemTypes.h
@@ -94,7 +94,7 @@ public:
 /**
 *	@struct FTargetingDefaultResultData
 */
-USTRUCT()
+USTRUCT(BlueprintType)
 struct FTargetingDefaultResultData
 {
 	GENERATED_BODY()
@@ -104,10 +104,11 @@ public:
 	{ }
 
 	/** The hit result for this target */
-	UPROPERTY()
+	UPROPERTY(BlueprintReadOnly, Category=Targeting)
 	FHitResult HitResult;
 
 	/** The score associated w/ this target */
+	UPROPERTY(BlueprintReadOnly, Category=Targeting)
 	float Score = 0.0f;
 };
 
