Fix CIS missing include

diff --git a/Engine/Plugins/Experimental/GameplayTargetingSystem/Source/GameplayTargetingSystem/Private/Tasks/SimpleTargetingSelectionTask.cpp b/Engine/Plugins/Experimental/GameplayTargetingSystem/Source/GameplayTargetingSystem/Private/Tasks/SimpleTargetingSelectionTask.cpp
--- a/Engine/Plugins/Experimental/GameplayTargetingSystem/Source/GameplayTargetingSystem/Private/Tasks/SimpleTargetingSelectionTask.cpp
+++ b/Engine/Plugins/Experimental/GameplayTargetingSystem/Source/GameplayTargetingSystem/Private/Tasks/SimpleTargetingSelectionTask.cpp
@@ -1,6 +1,7 @@
 ﻿// Copyright Epic Games, Inc. All Rights Reserved.
 
 #include "Tasks/SimpleTargetingSelectionTask.h"
+#include "GameFramework/Actor.h"
 
 void USimpleTargetingSelectionTask::Execute(const FTargetingRequestHandle& TargetingHandle) const
 {
