Allow TargetingTask blueprints to call functions that need a world context object and add a helper function for getting the TargetingSubsystem

diff --git a/Engine/Plugins/Experimental/GameplayTargetingSystem/Source/GameplayTargetingSystem/Private/TargetingSystem/TargetingSubsystem.cpp b/Engine/Plugins/Experimental/GameplayTargetingSystem/Source/GameplayTargetingSystem/Private/TargetingSystem/TargetingSubsystem.cpp
--- a/Engine/Plugins/Experimental/GameplayTargetingSystem/Source/GameplayTargetingSystem/Private/TargetingSystem/TargetingSubsystem.cpp
+++ b/Engine/Plugins/Experimental/GameplayTargetingSystem/Source/GameplayTargetingSystem/Private/TargetingSystem/TargetingSubsystem.cpp
@@ -355,7 +355,7 @@ void UTargetingSubsystem::ExecuteTargetingRequestWithHandle(FTargetingRequestHan
 	if (bTickingAsycnRequests)
 	{
 		FTargetingRequestData& RequestData = PendingTargetingRequests.FindOrAdd(TargetingHandle);
-		RequestData.Initialize(CompletionDelegate, CompletionDynamicDelegate);
+		RequestData.Initialize(CompletionDelegate, CompletionDynamicDelegate, this);
 		return;
 	}
 
@@ -393,7 +393,7 @@ void UTargetingSubsystem::ExecuteTargetingRequest(const UTargetingPreset* Target
 		if (bTickingAsycnRequests)
 		{
 			FTargetingRequestData& RequestData = PendingTargetingRequests.FindOrAdd(RequestHandle);
-			RequestData.Initialize(CompletionDelegate, CompletionDynamicDelegate);
+			RequestData.Initialize(CompletionDelegate, CompletionDynamicDelegate, this);
 			return;
 		}
 
@@ -408,7 +408,7 @@ void UTargetingSubsystem::StartAsyncTargetingRequestWithHandle(FTargetingRequest
 	if (bTickingAsycnRequests)
 	{
 		FTargetingRequestData& RequestData = PendingAsyncTargetingRequests.FindOrAdd(TargetingHandle);
-		RequestData.Initialize(CompletionDelegate, CompletionDynamicDelegate);
+		RequestData.Initialize(CompletionDelegate, CompletionDynamicDelegate, this);
 		return;
 	}
 
@@ -486,7 +486,7 @@ FTargetingRequestHandle UTargetingSubsystem::StartAsyncTargetingRequest(const UT
 		if (bTickingAsycnRequests)
 		{
 			FTargetingRequestData& RequestData = PendingAsyncTargetingRequests.FindOrAdd(RequestHandle);
-			RequestData.Initialize(CompletionDelegate, CompletionDynamicDelegate);
+			RequestData.Initialize(CompletionDelegate, CompletionDynamicDelegate, this);
 			return RequestHandle;
 		}
 
@@ -496,13 +496,13 @@ FTargetingRequestHandle UTargetingSubsystem::StartAsyncTargetingRequest(const UT
 	return RequestHandle;
 }
 
-void UTargetingSubsystem::ExecuteTargetingRequestWithHandleInternal(FTargetingRequestHandle TargetingHandle, FTargetingRequestDelegate CompletionDelegate, FTargetingRequestDynamicDelegate CompletionDynamicDelegate) const
+void UTargetingSubsystem::ExecuteTargetingRequestWithHandleInternal(FTargetingRequestHandle TargetingHandle, FTargetingRequestDelegate CompletionDelegate, FTargetingRequestDynamicDelegate CompletionDynamicDelegate)
 {
 	if (ensure(TargetingHandle.IsValid()))
 	{
 		// store the request data
 		FTargetingRequestData& RequestData = FTargetingRequestData::FindOrAdd(TargetingHandle);
-		RequestData.Initialize(CompletionDelegate, CompletionDynamicDelegate);
+		RequestData.Initialize(CompletionDelegate, CompletionDynamicDelegate, this);
 
 #if ENABLE_DRAW_DEBUG
 		if (IsTargetingDebugEnabled())
@@ -547,7 +547,7 @@ void UTargetingSubsystem::StartAsyncTargetingRequestWithHandleInternal(FTargetin
 {
 	// store the request data
 	FTargetingRequestData& RequestData = FTargetingRequestData::FindOrAdd(TargetingHandle);
-	RequestData.Initialize(CompletionDelegate, CompletionDynamicDelegate);
+	RequestData.Initialize(CompletionDelegate, CompletionDynamicDelegate, this);
 
 	// initialize the request data for async processing
 	FTargetingAsyncTaskData& AsyncTaskData = FTargetingAsyncTaskData::FindOrAdd(TargetingHandle);
diff --git a/Engine/Plugins/Experimental/GameplayTargetingSystem/Source/GameplayTargetingSystem/Private/Tasks/TargetingTask.cpp b/Engine/Plugins/Experimental/GameplayTargetingSystem/Source/GameplayTargetingSystem/Private/Tasks/TargetingTask.cpp
--- a/Engine/Plugins/Experimental/GameplayTargetingSystem/Source/GameplayTargetingSystem/Private/Tasks/TargetingTask.cpp
+++ b/Engine/Plugins/Experimental/GameplayTargetingSystem/Source/GameplayTargetingSystem/Private/Tasks/TargetingTask.cpp
@@ -3,6 +3,7 @@
 
 #include "GameFramework/Actor.h"
 #include "Types/TargetingSystemTypes.h"
+#include "TargetingSystem/TargetingSubsystem.h"
 
 
 UTargetingTask::UTargetingTask(const FObjectInitializer& ObjectInitializer)
@@ -94,3 +95,12 @@ UWorld* UTargetingTask::GetSourceContextWorld(const FTargetingRequestHandle& Tar
 
 	return GetWorld();
 }
+
+UTargetingSubsystem* UTargetingTask::GetTargetingSubsystem(const FTargetingRequestHandle& TargetingHandle) const
+{
+	if (const FTargetingRequestData* RequestData = FTargetingRequestData::Find(TargetingHandle))
+	{
+		return RequestData->TargetingSubsystem;
+	}
+	return nullptr;
+}
diff --git a/Engine/Plugins/Experimental/GameplayTargetingSystem/Source/GameplayTargetingSystem/Private/Types/TargetingSystemTypes.cpp b/Engine/Plugins/Experimental/GameplayTargetingSystem/Source/GameplayTargetingSystem/Private/Types/TargetingSystemTypes.cpp
--- a/Engine/Plugins/Experimental/GameplayTargetingSystem/Source/GameplayTargetingSystem/Private/Types/TargetingSystemTypes.cpp
+++ b/Engine/Plugins/Experimental/GameplayTargetingSystem/Source/GameplayTargetingSystem/Private/Types/TargetingSystemTypes.cpp
@@ -73,8 +73,9 @@ FTargetingRequestData* FTargetingRequestData::Find(FTargetingRequestHandle Handl
 	return UE::TargetingSystem::GTargetingRequestDataStore.Find(Handle);
 }
 
-void FTargetingRequestData::Initialize(FTargetingRequestDelegate CompletionDelegate, FTargetingRequestDynamicDelegate CompletionDynamicDelegate)
+void FTargetingRequestData::Initialize(FTargetingRequestDelegate CompletionDelegate, FTargetingRequestDynamicDelegate CompletionDynamicDelegate, UTargetingSubsystem* Subsystem)
 {
+	TargetingSubsystem = Subsystem;
 	bComplete = false;
 	TargetingRequestDelegate = CompletionDelegate;
 	TargetingRequestDynamicDelegate = CompletionDynamicDelegate;
diff --git a/Engine/Plugins/Experimental/GameplayTargetingSystem/Source/GameplayTargetingSystem/Public/TargetingSystem/TargetingSubsystem.h b/Engine/Plugins/Experimental/GameplayTargetingSystem/Source/GameplayTargetingSystem/Public/TargetingSystem/TargetingSubsystem.h
--- a/Engine/Plugins/Experimental/GameplayTargetingSystem/Source/GameplayTargetingSystem/Public/TargetingSystem/TargetingSubsystem.h
+++ b/Engine/Plugins/Experimental/GameplayTargetingSystem/Source/GameplayTargetingSystem/Public/TargetingSystem/TargetingSubsystem.h
@@ -151,7 +151,7 @@ public:
 
 private:
 	/** Internal method to handle running an immediate targeting request */
-	void ExecuteTargetingRequestWithHandleInternal(FTargetingRequestHandle TargetingHandle, FTargetingRequestDelegate CompletionDelegate, FTargetingRequestDynamicDelegate CompletionDynamicDelegate) const;
+	void ExecuteTargetingRequestWithHandleInternal(FTargetingRequestHandle TargetingHandle, FTargetingRequestDelegate CompletionDelegate, FTargetingRequestDynamicDelegate CompletionDynamicDelegate);
 
 	/** Internal method to hanlde queuing up an async targeting request */
 	void StartAsyncTargetingRequestWithHandleInternal(FTargetingRequestHandle TargetingHandle, FTargetingRequestDelegate CompletionDelegate, FTargetingRequestDynamicDelegate CompletionDynamicDelegate);
diff --git a/Engine/Plugins/Experimental/GameplayTargetingSystem/Source/GameplayTargetingSystem/Public/Tasks/TargetingTask.h b/Engine/Plugins/Experimental/GameplayTargetingSystem/Source/GameplayTargetingSystem/Public/Tasks/TargetingTask.h
--- a/Engine/Plugins/Experimental/GameplayTargetingSystem/Source/GameplayTargetingSystem/Public/Tasks/TargetingTask.h
+++ b/Engine/Plugins/Experimental/GameplayTargetingSystem/Source/GameplayTargetingSystem/Public/Tasks/TargetingTask.h
@@ -40,7 +40,7 @@ struct FTargetingDebugInfo;
 *	the end user might prefer to make decisions on. Distance (min/max), score rating
 *	etc.
 */
-UCLASS(EditInlineNew, Abstract, Const)
+UCLASS(EditInlineNew, Abstract, Const, meta=(ShowWorldContextPin="true"))
 class TARGETINGSYSTEM_API UTargetingTask : public UObject
 {
 	GENERATED_BODY()
@@ -70,6 +70,10 @@ protected:
 	/** Helper method to get the world from the source context (if possible, returns nullptr if one cannot be found) */
 	UWorld* GetSourceContextWorld(const FTargetingRequestHandle& TargetingHandle) const;
 
+	/** Helper method to get the Targeting Subsystem in TargetingTask Blueprint Types */
+	UFUNCTION(BlueprintPure, Category="Targeting")
+	UTargetingSubsystem* GetTargetingSubsystem(const FTargetingRequestHandle& TargetingHandle) const;
+
 	/** Debug Helper Methods */
 #if ENABLE_DRAW_DEBUG
 public:
diff --git a/Engine/Plugins/Experimental/GameplayTargetingSystem/Source/GameplayTargetingSystem/Public/Types/TargetingSystemTypes.h b/Engine/Plugins/Experimental/GameplayTargetingSystem/Source/GameplayTargetingSystem/Public/Types/TargetingSystemTypes.h
--- a/Engine/Plugins/Experimental/GameplayTargetingSystem/Source/GameplayTargetingSystem/Public/Types/TargetingSystemTypes.h
+++ b/Engine/Plugins/Experimental/GameplayTargetingSystem/Source/GameplayTargetingSystem/Public/Types/TargetingSystemTypes.h
@@ -17,6 +17,7 @@ class FArchive;
 class UTargetingPreset;
 class UTargetingTask;
 class UPackageMap;
+class UTargetingSubsystem;
 
 
 /**
@@ -200,11 +201,14 @@ public:
 	TARGETINGSYSTEM_API static FTargetingRequestData* Find(FTargetingRequestHandle Handle);
 
 	/** Initializes the targeting request data for async processing */
-	void Initialize(FTargetingRequestDelegate CompletionDelegate, FTargetingRequestDynamicDelegate CompletionDynamicDelegate);
+	void Initialize(FTargetingRequestDelegate CompletionDelegate, FTargetingRequestDynamicDelegate CompletionDynamicDelegate, UTargetingSubsystem* Subsystem);
 
 	/** Broadcasts the targeting request delegate */
 	void BroadcastTargetingRequestDelegate(FTargetingRequestHandle TargetingRequestHandle);
 
+	/** Ref to the TargetingSubsystem object processing this request */
+	UTargetingSubsystem* TargetingSubsystem;
+
 	/** Indicates this handle has completed all the targeting request */
 	uint8 bComplete : 1;
 
