Create Dynamic Targeting Task (No new system needed)
#rnx

diff --git a/Engine/Plugins/Experimental/GameplayTargetingSystem/Source/GameplayTargetingSystem/Public/TargetingSystem/TargetingPreset.h b/Engine/Plugins/Experimental/GameplayTargetingSystem/Source/GameplayTargetingSystem/Public/TargetingSystem/TargetingPreset.h
--- a/Engine/Plugins/Experimental/GameplayTargetingSystem/Source/GameplayTargetingSystem/Public/TargetingSystem/TargetingPreset.h
+++ b/Engine/Plugins/Experimental/GameplayTargetingSystem/Source/GameplayTargetingSystem/Public/TargetingSystem/TargetingPreset.h
@@ -32,6 +32,8 @@ public:
 	/** Get the set of tasks defined for this targeting preset */
 	FORCEINLINE const FTargetingTaskSet* GetTargetingTaskSet() const { return &TargetingTaskSet; }
 	FORCEINLINE FTargetingTaskSet& GetMutableTargetingTaskSet() { return TargetingTaskSet; }
+	UFUNCTION(BlueprintCallable, Category="Targeting")
+	FORCEINLINE TArray<UTargetingTask*> GetTasks() { return TargetingTaskSet.Tasks; }
 
 private:
 	/** The tasks that make up this targeting preset */
