Fixed UTargetingPreset to return reference to the taskset array rather than copy

[Test Case]
1. Ensure that blueprint has a task being added dynamically
2. Have a debug string printed when the task is being executed
3. Dynamically attach the task to the preset and ensure that the the debug string is being printed while the targeting preset is run

diff --git a/Engine/Plugins/Experimental/GameplayTargetingSystem/Source/GameplayTargetingSystem/Public/TargetingSystem/TargetingPreset.h b/Engine/Plugins/Experimental/GameplayTargetingSystem/Source/GameplayTargetingSystem/Public/TargetingSystem/TargetingPreset.h
--- a/Engine/Plugins/Experimental/GameplayTargetingSystem/Source/GameplayTargetingSystem/Public/TargetingSystem/TargetingPreset.h
+++ b/Engine/Plugins/Experimental/GameplayTargetingSystem/Source/GameplayTargetingSystem/Public/TargetingSystem/TargetingPreset.h
@@ -32,11 +32,8 @@ public:
 	/** Get the set of tasks defined for this targeting preset */
 	FORCEINLINE const FTargetingTaskSet* GetTargetingTaskSet() const { return &TargetingTaskSet; }
 	FORCEINLINE FTargetingTaskSet& GetMutableTargetingTaskSet() { return TargetingTaskSet; }
-	UFUNCTION(BlueprintCallable, Category="Targeting")
-	FORCEINLINE TArray<UTargetingTask*> GetTasks() { return TargetingTaskSet.Tasks; }
 
-private:
 	/** The tasks that make up this targeting preset */
-	UPROPERTY(EditAnywhere, Category = "Targeting", meta = (ShowOnlyInnerProperties))
+	UPROPERTY(EditAnywhere, BlueprintReadWrite, Category = "Targeting", meta = (ShowOnlyInnerProperties))
 	FTargetingTaskSet TargetingTaskSet;
 };
\ No newline at end of file
diff --git a/Engine/Plugins/Experimental/GameplayTargetingSystem/Source/GameplayTargetingSystem/Public/Types/TargetingSystemTypes.h b/Engine/Plugins/Experimental/GameplayTargetingSystem/Source/GameplayTargetingSystem/Public/Types/TargetingSystemTypes.h
--- a/Engine/Plugins/Experimental/GameplayTargetingSystem/Source/GameplayTargetingSystem/Public/Types/TargetingSystemTypes.h
+++ b/Engine/Plugins/Experimental/GameplayTargetingSystem/Source/GameplayTargetingSystem/Public/Types/TargetingSystemTypes.h
@@ -87,7 +87,7 @@ public:
 	TARGETINGSYSTEM_API static const FTargetingTaskSet** Find(FTargetingRequestHandle Handle);
 
 	/** The set of tasks that will be used to satisfy a targeting request  */
-	UPROPERTY(EditAnywhere, Instanced, Category="Targeting Task Set")
+	UPROPERTY(EditAnywhere, BlueprintReadWrite, Instanced, Category="Targeting Task Set")
 	TArray<TObjectPtr<UTargetingTask>> Tasks;
 };
 
