Fix server crash from AsyncAction_PerformTargeting that fails to properly clean itself up

diff --git a/Engine/Plugins/Experimental/GameplayTargetingSystem/Source/GameplayTargetingSystem/Private/Async/AsyncAction_PerformTargeting.cpp b/Engine/Plugins/Experimental/GameplayTargetingSystem/Source/GameplayTargetingSystem/Private/Async/AsyncAction_PerformTargeting.cpp
--- a/Engine/Plugins/Experimental/GameplayTargetingSystem/Source/GameplayTargetingSystem/Private/Async/AsyncAction_PerformTargeting.cpp
+++ b/Engine/Plugins/Experimental/GameplayTargetingSystem/Source/GameplayTargetingSystem/Private/Async/AsyncAction_PerformTargeting.cpp
@@ -82,6 +82,7 @@ void UAsyncAction_PerformTargeting::Activate()
 					TARGETING_LOG(Verbose, TEXT("Entering request lambda"));
 
 					Targeted.Broadcast(InTargetingHandle);
+					SetReadyToDestroy();
 				});
 
 			if (bUseAsyncTargeting)
@@ -97,12 +98,10 @@ void UAsyncAction_PerformTargeting::Activate()
 			{
 				TARGETING_LOG(Verbose, TEXT("Starting immediate targeting"));
 
-				FTargetingImmediateTaskData& ImmeidateTaskData = FTargetingImmediateTaskData::FindOrAdd(TargetingHandle);
-				ImmeidateTaskData.bReleaseOnCompletion = true;
+				FTargetingImmediateTaskData& ImmediateTaskData = FTargetingImmediateTaskData::FindOrAdd(TargetingHandle);
+				ImmediateTaskData.bReleaseOnCompletion = true;
 
 				TargetingSubsystem->ExecuteTargetingRequestWithHandle(TargetingHandle, Delegate);
-
-				SetReadyToDestroy();
 			}
 		}
 	}
