Export the targeting system's GetOverrideTargetingLifeTime so plugins could call it.

#rnx

diff --git a/Engine/Plugins/Experimental/GameplayTargetingSystem/Source/GameplayTargetingSystem/Public/TargetingSystem/TargetingSubsystem.h b/Engine/Plugins/Experimental/GameplayTargetingSystem/Source/GameplayTargetingSystem/Public/TargetingSystem/TargetingSubsystem.h
--- a/Engine/Plugins/Experimental/GameplayTargetingSystem/Source/GameplayTargetingSystem/Public/TargetingSystem/TargetingSubsystem.h
+++ b/Engine/Plugins/Experimental/GameplayTargetingSystem/Source/GameplayTargetingSystem/Public/TargetingSystem/TargetingSubsystem.h
@@ -208,7 +208,7 @@ public:
 	TARGETINGSYSTEM_API static bool IsTargetingDebugEnabled();
 
 	/** Helper method to overrides the life time of the shapes drawn */
-	static float GetOverrideTargetingLifeTime();
+	TARGETINGSYSTEM_API static float GetOverrideTargetingLifeTime();
 
 	/** Callback to handle OnShowDebugInfo to check if the TargetingSystem has been enabled */
 	static void OnShowDebugInfo(AHUD* HUD, UCanvas* Canvas, const FDebugDisplayInfo& DisplayInfo, float& YL, float& YPos);
