Fix use-after-free: copy the RequestData before calling the delegates, as the delegates could invalidate the RequestData reference by calling ExecuteTargetingRequestWithHandleInternal.

#tests 45 minutes using address sanitizer

diff --git a/Engine/Plugins/Experimental/GameplayTargetingSystem/Source/GameplayTargetingSystem/Private/TargetingSystem/TargetingSubsystem.cpp b/Engine/Plugins/Experimental/GameplayTargetingSystem/Source/GameplayTargetingSystem/Private/TargetingSystem/TargetingSubsystem.cpp
--- a/Engine/Plugins/Experimental/GameplayTargetingSystem/Source/GameplayTargetingSystem/Private/TargetingSystem/TargetingSubsystem.cpp
+++ b/Engine/Plugins/Experimental/GameplayTargetingSystem/Source/GameplayTargetingSystem/Private/TargetingSystem/TargetingSubsystem.cpp
@@ -527,7 +527,10 @@ void UTargetingSubsystem::ExecuteTargetingRequestWithHandleInternal(FTargetingRe
 				}
 
 				RequestData.bComplete = true;
-				RequestData.BroadcastTargetingRequestDelegate(TargetingHandle);
+				/* Creates a copy of the RequestData, so that the delegates are free to call ExecuteTargetingRequestWithHandleInternal recursively.
+				   After this copy, the RequestData reference is considered as invalid. */
+				FTargetingRequestData RequestDataCopy = RequestData;
+				RequestDataCopy.BroadcastTargetingRequestDelegate(TargetingHandle);
 
 #if ENABLE_DRAW_DEBUG
 #if WITH_EDITORONLY_DATA
