Fix use-after-free at shutdown: fix initialization order to ensure the delegate lives long enough.

#tests shutdown of the game in various game modes using address sanitizer
#localization none

diff --git a/Engine/Plugins/Experimental/GameplayTargetingSystem/Source/GameplayTargetingSystem/Private/TargetingSystem/TargetingSubsystem.cpp b/Engine/Plugins/Experimental/GameplayTargetingSystem/Source/GameplayTargetingSystem/Private/TargetingSystem/TargetingSubsystem.cpp
--- a/Engine/Plugins/Experimental/GameplayTargetingSystem/Source/GameplayTargetingSystem/Private/TargetingSystem/TargetingSubsystem.cpp
+++ b/Engine/Plugins/Experimental/GameplayTargetingSystem/Source/GameplayTargetingSystem/Private/TargetingSystem/TargetingSubsystem.cpp
@@ -75,7 +75,11 @@ namespace TargetingSystemCVars
 #endif // ENABLE_DRAW_DEBUG
 }
 
-UTargetingSubsystem::FOnTargetingRequestHandleReleased UTargetingSubsystem::ReleaseHandleDelegate;
+UTargetingSubsystem::FOnTargetingRequestHandleReleased& UTargetingSubsystem::ReleaseHandleDelegate()
+{
+	static UTargetingSubsystem::FOnTargetingRequestHandleReleased Delegate;
+	return Delegate;
+}
 
 
 UTargetingSubsystem::UTargetingSubsystem()
@@ -345,7 +349,7 @@ void UTargetingSubsystem::ReleaseTargetRequestHandle(FTargetingRequestHandle& Ha
 	}
 #endif // ENABLE_DRAW_DEBUG
 
-	ReleaseHandleDelegate.Broadcast(CachedHandle);
+	ReleaseHandleDelegate().Broadcast(CachedHandle);
 	TARGETING_LOG(Verbose, TEXT("%s: - Releasigng Handle [%d]"), ANSI_TO_TCHAR(__FUNCTION__), CachedHandle.Handle);
 }
 
diff --git a/Engine/Plugins/Experimental/GameplayTargetingSystem/Source/GameplayTargetingSystem/Public/TargetingSystem/TargetingSubsystem.h b/Engine/Plugins/Experimental/GameplayTargetingSystem/Source/GameplayTargetingSystem/Public/TargetingSystem/TargetingSubsystem.h
--- a/Engine/Plugins/Experimental/GameplayTargetingSystem/Source/GameplayTargetingSystem/Public/TargetingSystem/TargetingSubsystem.h
+++ b/Engine/Plugins/Experimental/GameplayTargetingSystem/Source/GameplayTargetingSystem/Public/TargetingSystem/TargetingSubsystem.h
@@ -125,7 +125,7 @@ public:
 
 	/** The handle released delegate that fires right before a handle is release so all data stores can clean up their state */
 	DECLARE_MULTICAST_DELEGATE_OneParam(FOnTargetingRequestHandleReleased, FTargetingRequestHandle)
-	TARGETINGSYSTEM_API static FOnTargetingRequestHandleReleased ReleaseHandleDelegate;
+	TARGETINGSYSTEM_API static FOnTargetingRequestHandleReleased& ReleaseHandleDelegate();
 
 	/** Target Handle Generation Methods */
 
diff --git a/Engine/Plugins/Experimental/GameplayTargetingSystem/Source/GameplayTargetingSystem/Public/Types/TargetingSystemDataStores.h b/Engine/Plugins/Experimental/GameplayTargetingSystem/Source/GameplayTargetingSystem/Public/Types/TargetingSystemDataStores.h
--- a/Engine/Plugins/Experimental/GameplayTargetingSystem/Source/GameplayTargetingSystem/Public/Types/TargetingSystemDataStores.h
+++ b/Engine/Plugins/Experimental/GameplayTargetingSystem/Source/GameplayTargetingSystem/Public/Types/TargetingSystemDataStores.h
@@ -28,14 +28,14 @@ namespace UE
 		{
 			TTargetingDataStore()
 			{
-				ReleaseDelegateHandle = UTargetingSubsystem::ReleaseHandleDelegate.AddStatic(&TTargetingDataStore::OnTargetingRequestHandleReleased);
+				ReleaseDelegateHandle = UTargetingSubsystem::ReleaseHandleDelegate().AddStatic(&TTargetingDataStore::OnTargetingRequestHandleReleased);
 			}
 
 			~TTargetingDataStore()
 			{
 				if (ReleaseDelegateHandle.IsValid())
 				{
-					UTargetingSubsystem::ReleaseHandleDelegate.Remove(ReleaseDelegateHandle);
+					UTargetingSubsystem::ReleaseHandleDelegate().Remove(ReleaseDelegateHandle);
 				}
 			}
 
