TargetingSubsystem: Adding UPARAM(ref) for the TargetingHandle parameter so that, on the blueprint side, the function RemoveAsyncTargetingRequestWithHandle() takes the param as an input instead of using it as an output.
#rnx

diff --git a/Engine/Plugins/Experimental/GameplayTargetingSystem/Source/GameplayTargetingSystem/Public/TargetingSystem/TargetingSubsystem.h b/Engine/Plugins/Experimental/GameplayTargetingSystem/Source/GameplayTargetingSystem/Public/TargetingSystem/TargetingSubsystem.h
--- a/Engine/Plugins/Experimental/GameplayTargetingSystem/Source/GameplayTargetingSystem/Public/TargetingSystem/TargetingSubsystem.h
+++ b/Engine/Plugins/Experimental/GameplayTargetingSystem/Source/GameplayTargetingSystem/Public/TargetingSystem/TargetingSubsystem.h
@@ -143,7 +143,7 @@ public:
 
 	/** Method to remove an async targeting request with a given targeting handle */
 	UFUNCTION(BlueprintCallable, Category = "Targeting System | Async Request")
-	TARGETINGSYSTEM_API void RemoveAsyncTargetingRequestWithHandle(FTargetingRequestHandle& TargetingHandle);
+	TARGETINGSYSTEM_API void RemoveAsyncTargetingRequestWithHandle(UPARAM(ref) FTargetingRequestHandle& TargetingHandle);
 
 	/** Method to queue an async targeting request based on a gameplay targeting preset. */
 	UFUNCTION(BlueprintCallable, Category = "Targeting System | Async Request")
