AsyncAction_PerformTargeting: Adding meta BlueprintType so that UAsyncAction_PerformTargeting* can be stored as a variable in blueprints.
#rnx

diff --git a/Engine/Plugins/Experimental/GameplayTargetingSystem/Source/GameplayTargetingSystem/Public/Async/AsyncAction_PerformTargeting.h b/Engine/Plugins/Experimental/GameplayTargetingSystem/Source/GameplayTargetingSystem/Public/Async/AsyncAction_PerformTargeting.h
--- a/Engine/Plugins/Experimental/GameplayTargetingSystem/Source/GameplayTargetingSystem/Public/Async/AsyncAction_PerformTargeting.h
+++ b/Engine/Plugins/Experimental/GameplayTargetingSystem/Source/GameplayTargetingSystem/Public/Async/AsyncAction_PerformTargeting.h
@@ -18,7 +18,7 @@ DECLARE_DYNAMIC_MULTICAST_DELEGATE_OneParam(FPerformTargetingReady, FTargetingRe
 /**
 *	@class UAsyncAction_PerformTargeting
 */
-UCLASS(meta = (ExposedAsyncProxy = "AsyncTaskRef"))
+UCLASS(BlueprintType, meta = (ExposedAsyncProxy = "AsyncTaskRef"))
 class UAsyncAction_PerformTargeting : public UBlueprintAsyncActionBase
 {
 	GENERATED_UCLASS_BODY()
