Target Selection Task: Added the ability to perform target traces using capsules and boxes (before only lines and spheres were available.)

- Fixed bug where the DefaultSweptTraceRadius was editable for traces where it was non-applicable.

[RN]

diff --git a/Engine/Config/BaseEngine.ini b/Engine/Config/BaseEngine.ini
--- a/Engine/Config/BaseEngine.ini
+++ b/Engine/Config/BaseEngine.ini
@@ -1219,6 +1219,8 @@ DefaultWorldPartitionHLODLayer=/Engine/Maps/Templates/HLODs/HLODLayer_Instanced.
 +FunctionRedirects=(OldName="EditorValidatorBase.CanValidate", NewName="EditorValidatorBase.K2_CanValidate")
 +FunctionRedirects=(OldName="EditorValidatorBase.CanValidateAsset", NewName="EditorValidatorBase.K2_CanValidateAsset")
 
++EnumRedirects=(OldName="ETargetingTraceType", ValueChanges=(("Sweep", "Sphere")) )
+
 ; End of releases redirects
 
 
diff --git a/Engine/Plugins/Experimental/GameplayTargetingSystem/Source/GameplayTargetingSystem/Private/Tasks/TargetingSelectionTask_Trace.cpp b/Engine/Plugins/Experimental/GameplayTargetingSystem/Source/GameplayTargetingSystem/Private/Tasks/TargetingSelectionTask_Trace.cpp
--- a/Engine/Plugins/Experimental/GameplayTargetingSystem/Source/GameplayTargetingSystem/Private/Tasks/TargetingSelectionTask_Trace.cpp
+++ b/Engine/Plugins/Experimental/GameplayTargetingSystem/Source/GameplayTargetingSystem/Private/Tasks/TargetingSelectionTask_Trace.cpp
@@ -94,10 +94,32 @@ float UTargetingSelectionTask_Trace::GetSweptTraceRadius_Implementation(const FT
 	return DefaultSweptTraceRadius.GetValue();
 }
 
+float UTargetingSelectionTask_Trace::GetSweptTraceCapsuleHalfHeight_Implementation(const FTargetingRequestHandle& TargetingHandle) const
+{
+	return DefaultSweptTraceCapsuleHalfHeight.GetValue();
+}
+
+FVector UTargetingSelectionTask_Trace::GetSweptTraceBoxHalfExtents_Implementation(const FTargetingRequestHandle& TargetingHandle) const
+{
+	return FVector(DefaultSweptTraceBoxHalfExtentX.GetValue(), DefaultSweptTraceBoxHalfExtentY.GetValue(), DefaultSweptTraceBoxHalfExtentZ.GetValue());
+}
+
+FRotator UTargetingSelectionTask_Trace::GetSweptTraceRotation_Implementation(const FTargetingRequestHandle& TargetingHandle) const
+{
+	return DefaultSweptTraceRotation;
+}
+
 void UTargetingSelectionTask_Trace::GetAdditionalActorsToIgnore_Implementation(const FTargetingRequestHandle& TargetingHandle, TArray<AActor*>& OutAdditionalActorsToIgnore) const
 {
 }
 
+FQuat UTargetingSelectionTask_Trace::GetSweptTraceQuat(const FVector& TraceDirection, const FTargetingRequestHandle& TargetingHandle) const
+{
+	const FQuat AbsoluteQuat = TraceDirection.ToOrientationQuat();
+	const FQuat RelativeQuat = GetSweptTraceRotation(TargetingHandle).Quaternion();
+	return AbsoluteQuat * RelativeQuat;
+}
+
 void UTargetingSelectionTask_Trace::ExecuteImmediateTrace(const FTargetingRequestHandle& TargetingHandle) const
 {
 	if (UWorld* World = GetSourceContextWorld(TargetingHandle))
@@ -106,8 +128,9 @@ void UTargetingSelectionTask_Trace::ExecuteImmediateTrace(const FTargetingReques
 		ResetTraceResultsDebugString(TargetingHandle);
 #endif // ENABLE_DRAW_DEBUG
 
+		const FVector Direction = GetTraceDirection(TargetingHandle).GetSafeNormal();
 		const FVector Start = (GetSourceLocation(TargetingHandle) + GetSourceOffset(TargetingHandle));
-		const FVector End = Start + (GetTraceDirection(TargetingHandle) * GetTraceLength(TargetingHandle));
+		const FVector End = Start + (Direction * GetTraceLength(TargetingHandle));
 
 		FCollisionQueryParams Params(SCENE_QUERY_STAT(ExecuteImmediateTrace), bComplexTrace);
 		InitCollisionParams(TargetingHandle, Params);
@@ -115,25 +138,47 @@ void UTargetingSelectionTask_Trace::ExecuteImmediateTrace(const FTargetingReques
 		TArray<FHitResult> Hits;
 		if (CollisionProfileName.Name != TEXT("NoCollision"))
 		{
-			if (TraceType == ETargetingTraceType::Sweep)
+			switch (TraceType)
 			{
+			default:
+			case ETargetingTraceType::Line:
+				World->LineTraceMultiByProfile(Hits, Start, End, CollisionProfileName.Name, Params);
+				break;
+			case ETargetingTraceType::Sphere:
 				World->SweepMultiByProfile(Hits, Start, End, FQuat::Identity, CollisionProfileName.Name, FCollisionShape::MakeSphere(GetSweptTraceRadius(TargetingHandle)), Params);
-			}
-			else
+				break;
+			case ETargetingTraceType::Capsule:
 			{
-				World->LineTraceMultiByProfile(Hits, Start, End, CollisionProfileName.Name, Params);
+				const FVector CapsuleShapeVector = FVector(0.0f, GetSweptTraceRadius(TargetingHandle), GetSweptTraceCapsuleHalfHeight(TargetingHandle));
+				World->SweepMultiByProfile(Hits, Start, End, GetSweptTraceQuat(Direction, TargetingHandle), CollisionProfileName.Name, FCollisionShape::MakeCapsule(CapsuleShapeVector), Params);
+			}
+				break;
+			case ETargetingTraceType::Box:
+				World->SweepMultiByProfile(Hits, Start, End, GetSweptTraceQuat(Direction, TargetingHandle), CollisionProfileName.Name, FCollisionShape::MakeBox(GetSweptTraceBoxHalfExtents(TargetingHandle)), Params);
+				break;
 			}
 		}
 		else
 		{
-			ECollisionChannel CollisionChannel = UEngineTypes::ConvertToCollisionChannel(TraceChannel);
-			if (TraceType == ETargetingTraceType::Sweep)
+			const ECollisionChannel CollisionChannel = UEngineTypes::ConvertToCollisionChannel(TraceChannel);
+			switch (TraceType)
 			{
+			default:
+			case ETargetingTraceType::Line:
+				World->LineTraceMultiByChannel(Hits, Start, End, CollisionChannel, Params);
+				break;
+			case ETargetingTraceType::Sphere:
 				World->SweepMultiByChannel(Hits, Start, End, FQuat::Identity, CollisionChannel, FCollisionShape::MakeSphere(GetSweptTraceRadius(TargetingHandle)), Params);
-			}
-			else
+				break;
+			case ETargetingTraceType::Capsule:
 			{
-				World->LineTraceMultiByChannel(Hits, Start, End, CollisionChannel, Params);
+				const FVector CapsuleShapeVector = FVector(0.0f, GetSweptTraceRadius(TargetingHandle), GetSweptTraceCapsuleHalfHeight(TargetingHandle));
+				World->SweepMultiByChannel(Hits, Start, End, GetSweptTraceQuat(Direction, TargetingHandle), CollisionChannel, FCollisionShape::MakeCapsule(CapsuleShapeVector), Params);
+			}
+				break;
+			case ETargetingTraceType::Box:
+				World->SweepMultiByChannel(Hits, Start, End, GetSweptTraceQuat(Direction, TargetingHandle), CollisionChannel, FCollisionShape::MakeBox(GetSweptTraceBoxHalfExtents(TargetingHandle)), Params);
+				break;
 			}
 		}
 
@@ -159,8 +204,9 @@ void UTargetingSelectionTask_Trace::ExecuteAsyncTrace(const FTargetingRequestHan
 		{
 			SourceActor = SourceContext->SourceActor;
 		}
+		const FVector Direction = GetTraceDirection(TargetingHandle).GetSafeNormal();
 		const FVector Start = (GetSourceLocation(TargetingHandle) + GetSourceOffset(TargetingHandle));
-		const FVector End = Start + (GetTraceDirection(TargetingHandle) * GetTraceLength(TargetingHandle));
+		const FVector End = Start + (Direction * GetTraceLength(TargetingHandle));
 
 		FCollisionQueryParams Params(SCENE_QUERY_STAT(ExecuteAsyncTrace), bComplexTrace);
 		InitCollisionParams(TargetingHandle, Params);
@@ -168,25 +214,47 @@ void UTargetingSelectionTask_Trace::ExecuteAsyncTrace(const FTargetingRequestHan
 		FTraceDelegate Delegate = FTraceDelegate::CreateUObject(this, &UTargetingSelectionTask_Trace::HandleAsyncTraceComplete, TargetingHandle);
 		if (CollisionProfileName.Name != TEXT("NoCollision"))
 		{
-			if (TraceType == ETargetingTraceType::Sweep)
+			switch (TraceType)
 			{
+			default:
+			case ETargetingTraceType::Line:
+				World->AsyncLineTraceByProfile(EAsyncTraceType::Multi, Start, End, CollisionProfileName.Name, Params, &Delegate);
+				break;
+			case ETargetingTraceType::Sphere:
 				World->AsyncSweepByProfile(EAsyncTraceType::Multi, Start, End, FQuat::Identity, CollisionProfileName.Name, FCollisionShape::MakeSphere(GetSweptTraceRadius(TargetingHandle)), Params, &Delegate);
-			}
-			else
+				break;
+			case ETargetingTraceType::Capsule:
 			{
-				World->AsyncLineTraceByProfile(EAsyncTraceType::Multi, Start, End, CollisionProfileName.Name, Params, &Delegate);
+				const FVector CapsuleShapeVector = FVector(0.0f, GetSweptTraceRadius(TargetingHandle), GetSweptTraceCapsuleHalfHeight(TargetingHandle));
+				World->AsyncSweepByProfile(EAsyncTraceType::Multi, Start, End, GetSweptTraceQuat(Direction, TargetingHandle), CollisionProfileName.Name, FCollisionShape::MakeCapsule(CapsuleShapeVector), Params, &Delegate);
+			}
+			break;
+			case ETargetingTraceType::Box:
+				World->AsyncSweepByProfile(EAsyncTraceType::Multi, Start, End, GetSweptTraceQuat(Direction, TargetingHandle), CollisionProfileName.Name, FCollisionShape::MakeBox(GetSweptTraceBoxHalfExtents(TargetingHandle)), Params, &Delegate);
+				break;
 			}
 		}
 		else
 		{
-			ECollisionChannel CollisionChannel = UEngineTypes::ConvertToCollisionChannel(TraceChannel);
-			if (TraceType == ETargetingTraceType::Sweep)
+			const ECollisionChannel CollisionChannel = UEngineTypes::ConvertToCollisionChannel(TraceChannel);
+			switch (TraceType)
 			{
+			default:
+			case ETargetingTraceType::Line:
+				World->AsyncLineTraceByChannel(EAsyncTraceType::Multi, Start, End, CollisionChannel, Params, FCollisionResponseParams::DefaultResponseParam, &Delegate);
+				break;
+			case ETargetingTraceType::Sphere:
 				World->AsyncSweepByChannel(EAsyncTraceType::Multi, Start, End, FQuat::Identity, CollisionChannel, FCollisionShape::MakeSphere(GetSweptTraceRadius(TargetingHandle)), Params, FCollisionResponseParams::DefaultResponseParam, &Delegate);
-			}
-			else
+				break;
+			case ETargetingTraceType::Capsule:
 			{
-				World->AsyncLineTraceByChannel(EAsyncTraceType::Multi, Start, End, CollisionChannel, Params, FCollisionResponseParams::DefaultResponseParam, &Delegate);
+				const FVector CapsuleShapeVector = FVector(0.0f, GetSweptTraceRadius(TargetingHandle), GetSweptTraceCapsuleHalfHeight(TargetingHandle));
+				World->AsyncSweepByChannel(EAsyncTraceType::Multi, Start, End, GetSweptTraceQuat(Direction, TargetingHandle), CollisionChannel, FCollisionShape::MakeCapsule(CapsuleShapeVector), Params, FCollisionResponseParams::DefaultResponseParam, &Delegate);
+			}
+				break;
+			case ETargetingTraceType::Box:
+				World->AsyncSweepByChannel(EAsyncTraceType::Multi, Start, End, GetSweptTraceQuat(Direction, TargetingHandle), CollisionChannel, FCollisionShape::MakeBox(GetSweptTraceBoxHalfExtents(TargetingHandle)), Params, FCollisionResponseParams::DefaultResponseParam, &Delegate);
+				break;
 			}
 		}
 	}
@@ -284,12 +352,30 @@ bool UTargetingSelectionTask_Trace::CanEditChange(const FProperty* InProperty) c
 
 	if (bCanEdit && InProperty)
 	{
-		if (InProperty->GetFName() == GET_MEMBER_NAME_CHECKED(UTargetingSelectionTask_Trace, DefaultSweptTraceRadius.GetValue()))
+		const FName PropertyName = InProperty->GetFName();
+		if (PropertyName == GET_MEMBER_NAME_CHECKED(UTargetingSelectionTask_Trace, DefaultSweptTraceRadius))
+		{
+			return (TraceType == ETargetingTraceType::Sphere || TraceType == ETargetingTraceType::Capsule);
+		}
+
+		if (PropertyName == GET_MEMBER_NAME_CHECKED(UTargetingSelectionTask_Trace, DefaultSweptTraceCapsuleHalfHeight))
+		{
+			return (TraceType == ETargetingTraceType::Capsule);
+		}
+
+		if (PropertyName == GET_MEMBER_NAME_CHECKED(UTargetingSelectionTask_Trace, DefaultSweptTraceBoxHalfExtentX)
+			|| PropertyName == GET_MEMBER_NAME_CHECKED(UTargetingSelectionTask_Trace, DefaultSweptTraceBoxHalfExtentY)
+			|| PropertyName == GET_MEMBER_NAME_CHECKED(UTargetingSelectionTask_Trace, DefaultSweptTraceBoxHalfExtentZ))
+		{
+			return (TraceType == ETargetingTraceType::Box);
+		}
+
+		if (PropertyName == GET_MEMBER_NAME_CHECKED(UTargetingSelectionTask_Trace, DefaultSweptTraceRotation))
 		{
-			return (TraceType == ETargetingTraceType::Sweep);
+			return (TraceType == ETargetingTraceType::Capsule || TraceType == ETargetingTraceType::Box);
 		}
 
-		if (InProperty->GetFName() == GET_MEMBER_NAME_CHECKED(UTargetingSelectionTask_Trace, TraceChannel))
+		if (PropertyName == GET_MEMBER_NAME_CHECKED(UTargetingSelectionTask_Trace, TraceChannel))
 		{
 			return (CollisionProfileName.Name == TEXT("NoCollision"));
 		}
diff --git a/Engine/Plugins/Experimental/GameplayTargetingSystem/Source/GameplayTargetingSystem/Public/Tasks/TargetingSelectionTask_Trace.h b/Engine/Plugins/Experimental/GameplayTargetingSystem/Source/GameplayTargetingSystem/Public/Tasks/TargetingSelectionTask_Trace.h
--- a/Engine/Plugins/Experimental/GameplayTargetingSystem/Source/GameplayTargetingSystem/Public/Tasks/TargetingSelectionTask_Trace.h
+++ b/Engine/Plugins/Experimental/GameplayTargetingSystem/Source/GameplayTargetingSystem/Public/Tasks/TargetingSelectionTask_Trace.h
@@ -18,15 +18,17 @@ struct FTargetingRequestHandle;
 struct FTraceDatum;
 struct FTraceHandle;
 
-
 /**
-*	@enum ETargetingAOEShape	
+*	@enum ETargetingTraceType determines the shape of the trace used.
 */
 UENUM()
 enum class ETargetingTraceType : uint8
 {
 	Line,
-	Sweep,
+	Sphere,
+	/* With zero rotation, the capsule length is along the cross product of the trace direction and Up. */
+	Capsule,
+	Box,
 };
 
 /**
@@ -63,10 +65,23 @@ protected:
 	UFUNCTION(BlueprintNativeEvent, Category = "Target Trace Selection")
 	float GetTraceLength(const FTargetingRequestHandle& TargetingHandle) const;
 
-	/** Native Event to get the swept trace radius (only called if bSweptTrace is true) */
+	/** Native Event to get the swept trace radius (only called if TraceType = ETargetingTraceType::Sphere or TraceType = ETargetingTraceType::Capsule) */
 	UFUNCTION(BlueprintNativeEvent, Category = "Target Trace Selection")
 	float GetSweptTraceRadius(const FTargetingRequestHandle& TargetingHandle) const;
 
+	/** Native Event to get the swept trace capsule's half height (only called if  TraceType = ETargetingTraceType::Capsule) */
+	UFUNCTION(BlueprintNativeEvent, Category = "Target Trace Selection")
+	float GetSweptTraceCapsuleHalfHeight(const FTargetingRequestHandle& TargetingHandle) const;
+
+	/** Native Event to get the swept box trace half extents (only called if TraceType = ETargetingTraceType::Box) */
+	UFUNCTION(BlueprintNativeEvent, Category = "Target Trace Selection")
+	FVector GetSweptTraceBoxHalfExtents(const FTargetingRequestHandle& TargetingHandle) const;
+
+	/** Native Event to get the swept trace rotation relative to trace direction
+	(only called if TraceType = ETargetingTraceType::Capsule or TraceType = ETargetingTraceType::Box) */
+	UFUNCTION(BlueprintNativeEvent, Category = "Target Trace Selection")
+	FRotator GetSweptTraceRotation(const FTargetingRequestHandle& TargetingHandle) const;
+
 	/** Native Event to get additional actors the Trace should ignore */
 	UFUNCTION(BlueprintNativeEvent, Category = "Target Trace Selection")
 	void GetAdditionalActorsToIgnore(const FTargetingRequestHandle& TargetingHandle, TArray<AActor*>& OutAdditionalActorsToIgnore) const;
@@ -87,8 +102,11 @@ private:
 	/** Setup CollisionQueryParams for the trace */
 	void InitCollisionParams(const FTargetingRequestHandle& TargetingHandle, FCollisionQueryParams& OutParams) const;
 
+	/** For non-sphere shape traces, calculates the world rotation for that trace. */
+	FQuat GetSweptTraceQuat(const FVector& TraceDirection, const FTargetingRequestHandle& TargetingHandle) const;
+
 protected:
-	/** The trace type to use */
+	/** The trace type (i.e. shape) to use */
 	UPROPERTY(EditAnywhere, Category = "Target Trace Selection | Collision Data")
 	ETargetingTraceType TraceType = ETargetingTraceType::Line;
 
@@ -100,9 +118,29 @@ protected:
 	UPROPERTY(EditAnywhere, Category = "Target Trace Selection | Collision Data")
 	FCollisionProfileName CollisionProfileName;
 
-	/** The default swept trace radius used by GetSweptTraceRadius when the trace type is set to Sweep */
+	/** The default swept trace radius used by GetSweptTraceRadius when the trace type is set to Sphere or Capsule */
 	UPROPERTY(EditAnywhere, Category = "Target Trace Selection | Swept Data")
 	FScalableFloat DefaultSweptTraceRadius = 10.0f;
+	
+	/** The default swept trace radius used by GetSweptTraceRadius when the trace type is set to Capsule */
+	UPROPERTY(EditAnywhere, Category = "Target Trace Selection | Swept Data")
+	FScalableFloat DefaultSweptTraceCapsuleHalfHeight = 10.0f;
+
+	/** The default swept extents used by GetSweptTraceExtents when the trace type is set to Box */
+	UPROPERTY(EditAnywhere, Category = "Target Trace Selection | Swept Data")
+	FScalableFloat DefaultSweptTraceBoxHalfExtentX = 10.0f;
+
+	/** The default swept extents used by GetSweptTraceBoxHalfExtents when the trace type is set to Box */
+	UPROPERTY(EditAnywhere, Category = "Target Trace Selection | Swept Data")
+	FScalableFloat DefaultSweptTraceBoxHalfExtentY = 10.0f;
+
+	/** The default swept extents used by GetSweptTraceBoxHalfExtents when the trace type is set to Box  */
+	UPROPERTY(EditAnywhere, Category = "Target Trace Selection | Swept Data")
+	FScalableFloat DefaultSweptTraceBoxHalfExtentZ = 10.0f;
+
+	/** The default swept rotation (relative to the trace direction) used by GetSweptTraceRotation when the trace type is set to Capsule or Box */
+	UPROPERTY(EditAnywhere, Category = "Target Trace Selection | Swept Data")
+	FRotator DefaultSweptTraceRotation = FRotator::ZeroRotator;
 
 	/** The default trace length to use if GetTraceLength is not overridden by a child */
 	UPROPERTY(EditAnywhere, Category = "Target Trace Selection | Trace Data")
