Register map changed delegate to clean up all references to UObjects which otherwise prevent garbage collection of the old map.

diff --git a/Engine/Plugins/Experimental/GameplayTargetingSystem/Source/GameplayTargetingSystem/Private/TargetingSystem/TargetingSubsystem.cpp b/Engine/Plugins/Experimental/GameplayTargetingSystem/Source/GameplayTargetingSystem/Private/TargetingSystem/TargetingSubsystem.cpp
--- a/Engine/Plugins/Experimental/GameplayTargetingSystem/Source/GameplayTargetingSystem/Private/TargetingSystem/TargetingSubsystem.cpp
+++ b/Engine/Plugins/Experimental/GameplayTargetingSystem/Source/GameplayTargetingSystem/Private/TargetingSystem/TargetingSubsystem.cpp
@@ -143,17 +143,13 @@ void UTargetingSubsystem::AddReferencedObjects(UObject* InThis, FReferenceCollec
 
 void UTargetingSubsystem::Initialize(FSubsystemCollectionBase& Collection)
 {
-
+	FCoreUObjectDelegates::PreLoadMap.AddUObject(this, &UTargetingSubsystem::HandlePreLoadMap);
 }
 
 void UTargetingSubsystem::Deinitialize()
 {
-	for (FTargetingRequestHandle AsyncHandle : AsyncTargetingRequests)
-	{
-		ReleaseTargetRequestHandle(AsyncHandle);
-	}
-
-	AsyncTargetingRequests.Empty();
+	ClearAsyncRequests();
+	FCoreUObjectDelegates::PreLoadMap.RemoveAll(this);
 }
 
 bool UTargetingSubsystem::Exec_Runtime(UWorld* Inworld, const TCHAR* Cmd, FOutputDevice& Ar)
@@ -755,6 +751,21 @@ UTargetingTask* UTargetingSubsystem::FindCurrentExecutingTask(FTargetingRequestH
 	return nullptr;
 }
 
+void UTargetingSubsystem::HandlePreLoadMap(const FString& MapName)
+{
+	/* clean up all async requests which will remove all references keeping the old world from being cleaned up */
+	ClearAsyncRequests();
+}
+
+void UTargetingSubsystem::ClearAsyncRequests()
+{
+	for (FTargetingRequestHandle AsyncHandle : AsyncTargetingRequests)
+	{
+		ReleaseTargetRequestHandle(AsyncHandle);
+	}
+
+	AsyncTargetingRequests.Empty();
+}
 
 #if ENABLE_DRAW_DEBUG
 
diff --git a/Engine/Plugins/Experimental/GameplayTargetingSystem/Source/GameplayTargetingSystem/Public/TargetingSystem/TargetingSubsystem.h b/Engine/Plugins/Experimental/GameplayTargetingSystem/Source/GameplayTargetingSystem/Public/TargetingSystem/TargetingSubsystem.h
--- a/Engine/Plugins/Experimental/GameplayTargetingSystem/Source/GameplayTargetingSystem/Public/TargetingSystem/TargetingSubsystem.h
+++ b/Engine/Plugins/Experimental/GameplayTargetingSystem/Source/GameplayTargetingSystem/Public/TargetingSystem/TargetingSubsystem.h
@@ -165,6 +165,13 @@ private:
 	/** Called when we set bTickingAsyncRequests to false, at this point it's safe to perform any queued operations on the Async Requests Array */
 	void OnFinishedTickingAsyncRequests();
 
+	/** Internal method to clear all async requests and release all references */
+	void ClearAsyncRequests();
+
+	/** Called when switching maps. Allows to release references that would otherwise prevent cleaning up an old world */
+	void HandlePreLoadMap(const FString& MapName);
+
+
 	/** The set of target requests queued up for async processing */
 	UPROPERTY(Transient)
 	TArray<FTargetingRequestHandle> AsyncTargetingRequests;
