Trace Targeting Debug Drawings: Improved debug drawings for the trace targeting system. Now the correct shapes are drawn and hits are drawn making the drawings much more informative.
- Created a new function signatures for DrawDebugCapsuleTraceSingle() and DrawDebugCapsuleTraceMulti() that take in an orientation for the capsule (before the functions would only draw in the orientation FQuat::Identity).

#rnx

diff --git a/Engine/Plugins/Experimental/GameplayTargetingSystem/Source/GameplayTargetingSystem/Private/Tasks/TargetingSelectionTask_Trace.cpp b/Engine/Plugins/Experimental/GameplayTargetingSystem/Source/GameplayTargetingSystem/Private/Tasks/TargetingSelectionTask_Trace.cpp
--- a/Engine/Plugins/Experimental/GameplayTargetingSystem/Source/GameplayTargetingSystem/Private/Tasks/TargetingSelectionTask_Trace.cpp
+++ b/Engine/Plugins/Experimental/GameplayTargetingSystem/Source/GameplayTargetingSystem/Private/Tasks/TargetingSelectionTask_Trace.cpp
@@ -7,6 +7,7 @@
 #include "Engine/World.h"
 #include "GameFramework/Actor.h"
 #include "GameFramework/Pawn.h"
+#include "KismetTraceUtils.h"
 #include "TargetingSystem/TargetingSubsystem.h"
 #include "Types/TargetingSystemLogs.h"
 
@@ -132,29 +133,37 @@ void UTargetingSelectionTask_Trace::ExecuteImmediateTrace(const FTargetingReques
 		const FVector Start = (GetSourceLocation(TargetingHandle) + GetSourceOffset(TargetingHandle));
 		const FVector End = Start + (Direction * GetTraceLength(TargetingHandle));
 
+		// Only bother calculating the orientation for shapes where orientation matters (i.e not points and not sphere)
+		FQuat OrientationQuat = FQuat::Identity;
+		if (TraceType != ETargetingTraceType::Line && TraceType != ETargetingTraceType::Sphere)
+		{
+			OrientationQuat = GetSweptTraceQuat(Direction, TargetingHandle);
+		}
+
 		FCollisionQueryParams Params(SCENE_QUERY_STAT(ExecuteImmediateTrace), bComplexTrace);
 		InitCollisionParams(TargetingHandle, Params);
 
+		bool bHasBlockingHit = false;
 		TArray<FHitResult> Hits;
 		if (CollisionProfileName.Name != TEXT("NoCollision"))
 		{
 			switch (TraceType)
 			{
-			default:
-			case ETargetingTraceType::Line:
-				World->LineTraceMultiByProfile(Hits, Start, End, CollisionProfileName.Name, Params);
-				break;
 			case ETargetingTraceType::Sphere:
-				World->SweepMultiByProfile(Hits, Start, End, FQuat::Identity, CollisionProfileName.Name, FCollisionShape::MakeSphere(GetSweptTraceRadius(TargetingHandle)), Params);
+				bHasBlockingHit = World->SweepMultiByProfile(Hits, Start, End, FQuat::Identity, CollisionProfileName.Name, FCollisionShape::MakeSphere(GetSweptTraceRadius(TargetingHandle)), Params);
 				break;
 			case ETargetingTraceType::Capsule:
 			{
 				const FVector CapsuleShapeVector = FVector(0.0f, GetSweptTraceRadius(TargetingHandle), GetSweptTraceCapsuleHalfHeight(TargetingHandle));
-				World->SweepMultiByProfile(Hits, Start, End, GetSweptTraceQuat(Direction, TargetingHandle), CollisionProfileName.Name, FCollisionShape::MakeCapsule(CapsuleShapeVector), Params);
+				bHasBlockingHit = World->SweepMultiByProfile(Hits, Start, End, OrientationQuat, CollisionProfileName.Name, FCollisionShape::MakeCapsule(CapsuleShapeVector), Params);
 			}
 				break;
 			case ETargetingTraceType::Box:
-				World->SweepMultiByProfile(Hits, Start, End, GetSweptTraceQuat(Direction, TargetingHandle), CollisionProfileName.Name, FCollisionShape::MakeBox(GetSweptTraceBoxHalfExtents(TargetingHandle)), Params);
+				bHasBlockingHit = World->SweepMultiByProfile(Hits, Start, End, OrientationQuat, CollisionProfileName.Name, FCollisionShape::MakeBox(GetSweptTraceBoxHalfExtents(TargetingHandle)), Params);
+				break;
+			default:
+			case ETargetingTraceType::Line:
+				bHasBlockingHit = World->LineTraceMultiByProfile(Hits, Start, End, CollisionProfileName.Name, Params);
 				break;
 			}
 		}
@@ -163,30 +172,27 @@ void UTargetingSelectionTask_Trace::ExecuteImmediateTrace(const FTargetingReques
 			const ECollisionChannel CollisionChannel = UEngineTypes::ConvertToCollisionChannel(TraceChannel);
 			switch (TraceType)
 			{
-			default:
-			case ETargetingTraceType::Line:
-				World->LineTraceMultiByChannel(Hits, Start, End, CollisionChannel, Params);
-				break;
 			case ETargetingTraceType::Sphere:
-				World->SweepMultiByChannel(Hits, Start, End, FQuat::Identity, CollisionChannel, FCollisionShape::MakeSphere(GetSweptTraceRadius(TargetingHandle)), Params);
+				bHasBlockingHit = World->SweepMultiByChannel(Hits, Start, End, FQuat::Identity, CollisionChannel, FCollisionShape::MakeSphere(GetSweptTraceRadius(TargetingHandle)), Params);
 				break;
 			case ETargetingTraceType::Capsule:
 			{
 				const FVector CapsuleShapeVector = FVector(0.0f, GetSweptTraceRadius(TargetingHandle), GetSweptTraceCapsuleHalfHeight(TargetingHandle));
-				World->SweepMultiByChannel(Hits, Start, End, GetSweptTraceQuat(Direction, TargetingHandle), CollisionChannel, FCollisionShape::MakeCapsule(CapsuleShapeVector), Params);
+				bHasBlockingHit = World->SweepMultiByChannel(Hits, Start, End, OrientationQuat, CollisionChannel, FCollisionShape::MakeCapsule(CapsuleShapeVector), Params);
 			}
 				break;
 			case ETargetingTraceType::Box:
-				World->SweepMultiByChannel(Hits, Start, End, GetSweptTraceQuat(Direction, TargetingHandle), CollisionChannel, FCollisionShape::MakeBox(GetSweptTraceBoxHalfExtents(TargetingHandle)), Params);
+				bHasBlockingHit = World->SweepMultiByChannel(Hits, Start, End, OrientationQuat, CollisionChannel, FCollisionShape::MakeBox(GetSweptTraceBoxHalfExtents(TargetingHandle)), Params);
+				break;
+			default:
+			case ETargetingTraceType::Line:
+				bHasBlockingHit = World->LineTraceMultiByChannel(Hits, Start, End, CollisionChannel, Params);
 				break;
 			}
 		}
 
 #if ENABLE_DRAW_DEBUG
-		if (UTargetingSubsystem::IsTargetingDebugEnabled())
-		{
-			DrawDebugLine(World, Start, End, FColor::Green, false, 30.0f, 0, 2.0f);
-		}
+		DrawDebugTrace(TargetingHandle, Start, End, OrientationQuat, bHasBlockingHit, Hits);
 #endif // ENABLE_DRAW_DEBUG
 
 		ProcessHitResults(TargetingHandle, Hits);
@@ -208,6 +214,13 @@ void UTargetingSelectionTask_Trace::ExecuteAsyncTrace(const FTargetingRequestHan
 		const FVector Start = (GetSourceLocation(TargetingHandle) + GetSourceOffset(TargetingHandle));
 		const FVector End = Start + (Direction * GetTraceLength(TargetingHandle));
 
+		// Only bother calculating the orientation for shapes where orientation matters (i.e not points and not sphere)
+		FQuat OrientationQuat = FQuat::Identity;
+		if (TraceType != ETargetingTraceType::Line && TraceType != ETargetingTraceType::Sphere)
+		{
+			OrientationQuat = GetSweptTraceQuat(Direction, TargetingHandle);
+		}
+
 		FCollisionQueryParams Params(SCENE_QUERY_STAT(ExecuteAsyncTrace), bComplexTrace);
 		InitCollisionParams(TargetingHandle, Params);
 
@@ -216,21 +229,21 @@ void UTargetingSelectionTask_Trace::ExecuteAsyncTrace(const FTargetingRequestHan
 		{
 			switch (TraceType)
 			{
-			default:
-			case ETargetingTraceType::Line:
-				World->AsyncLineTraceByProfile(EAsyncTraceType::Multi, Start, End, CollisionProfileName.Name, Params, &Delegate);
-				break;
 			case ETargetingTraceType::Sphere:
 				World->AsyncSweepByProfile(EAsyncTraceType::Multi, Start, End, FQuat::Identity, CollisionProfileName.Name, FCollisionShape::MakeSphere(GetSweptTraceRadius(TargetingHandle)), Params, &Delegate);
 				break;
 			case ETargetingTraceType::Capsule:
 			{
 				const FVector CapsuleShapeVector = FVector(0.0f, GetSweptTraceRadius(TargetingHandle), GetSweptTraceCapsuleHalfHeight(TargetingHandle));
-				World->AsyncSweepByProfile(EAsyncTraceType::Multi, Start, End, GetSweptTraceQuat(Direction, TargetingHandle), CollisionProfileName.Name, FCollisionShape::MakeCapsule(CapsuleShapeVector), Params, &Delegate);
+				World->AsyncSweepByProfile(EAsyncTraceType::Multi, Start, End, OrientationQuat, CollisionProfileName.Name, FCollisionShape::MakeCapsule(CapsuleShapeVector), Params, &Delegate);
 			}
 			break;
 			case ETargetingTraceType::Box:
-				World->AsyncSweepByProfile(EAsyncTraceType::Multi, Start, End, GetSweptTraceQuat(Direction, TargetingHandle), CollisionProfileName.Name, FCollisionShape::MakeBox(GetSweptTraceBoxHalfExtents(TargetingHandle)), Params, &Delegate);
+				World->AsyncSweepByProfile(EAsyncTraceType::Multi, Start, End, OrientationQuat, CollisionProfileName.Name, FCollisionShape::MakeBox(GetSweptTraceBoxHalfExtents(TargetingHandle)), Params, &Delegate);
+				break;
+			default:
+			case ETargetingTraceType::Line:
+				World->AsyncLineTraceByProfile(EAsyncTraceType::Multi, Start, End, CollisionProfileName.Name, Params, &Delegate);
 				break;
 			}
 		}
@@ -239,21 +252,21 @@ void UTargetingSelectionTask_Trace::ExecuteAsyncTrace(const FTargetingRequestHan
 			const ECollisionChannel CollisionChannel = UEngineTypes::ConvertToCollisionChannel(TraceChannel);
 			switch (TraceType)
 			{
-			default:
-			case ETargetingTraceType::Line:
-				World->AsyncLineTraceByChannel(EAsyncTraceType::Multi, Start, End, CollisionChannel, Params, FCollisionResponseParams::DefaultResponseParam, &Delegate);
-				break;
 			case ETargetingTraceType::Sphere:
 				World->AsyncSweepByChannel(EAsyncTraceType::Multi, Start, End, FQuat::Identity, CollisionChannel, FCollisionShape::MakeSphere(GetSweptTraceRadius(TargetingHandle)), Params, FCollisionResponseParams::DefaultResponseParam, &Delegate);
 				break;
 			case ETargetingTraceType::Capsule:
 			{
 				const FVector CapsuleShapeVector = FVector(0.0f, GetSweptTraceRadius(TargetingHandle), GetSweptTraceCapsuleHalfHeight(TargetingHandle));
-				World->AsyncSweepByChannel(EAsyncTraceType::Multi, Start, End, GetSweptTraceQuat(Direction, TargetingHandle), CollisionChannel, FCollisionShape::MakeCapsule(CapsuleShapeVector), Params, FCollisionResponseParams::DefaultResponseParam, &Delegate);
+				World->AsyncSweepByChannel(EAsyncTraceType::Multi, Start, End, OrientationQuat, CollisionChannel, FCollisionShape::MakeCapsule(CapsuleShapeVector), Params, FCollisionResponseParams::DefaultResponseParam, &Delegate);
 			}
 				break;
 			case ETargetingTraceType::Box:
-				World->AsyncSweepByChannel(EAsyncTraceType::Multi, Start, End, GetSweptTraceQuat(Direction, TargetingHandle), CollisionChannel, FCollisionShape::MakeBox(GetSweptTraceBoxHalfExtents(TargetingHandle)), Params, FCollisionResponseParams::DefaultResponseParam, &Delegate);
+				World->AsyncSweepByChannel(EAsyncTraceType::Multi, Start, End, OrientationQuat, CollisionChannel, FCollisionShape::MakeBox(GetSweptTraceBoxHalfExtents(TargetingHandle)), Params, FCollisionResponseParams::DefaultResponseParam, &Delegate);
+				break;
+			default:
+			case ETargetingTraceType::Line:
+				World->AsyncLineTraceByChannel(EAsyncTraceType::Multi, Start, End, CollisionChannel, Params, FCollisionResponseParams::DefaultResponseParam, &Delegate);
 				break;
 			}
 		}
@@ -271,13 +284,19 @@ void UTargetingSelectionTask_Trace::HandleAsyncTraceComplete(const FTraceHandle&
 #if ENABLE_DRAW_DEBUG
 		ResetTraceResultsDebugString(TargetingHandle);
 
-		if (UTargetingSubsystem::IsTargetingDebugEnabled())
+		// We have to manually find if there is a blocking hit.
+		bool bHasBlockingHit = false;
+		for (const FHitResult& HitResult : InTraceDatum.OutHits)
 		{
-			if (UWorld* World = GetSourceContextWorld(TargetingHandle))
+			if (HitResult.bBlockingHit)
 			{
-				DrawDebugLine(World, InTraceDatum.Start, InTraceDatum.End, FColor::Green, false, 30.0f, 0, 2.0f);
+				bHasBlockingHit = true;
+				break;
 			}
 		}
+
+		DrawDebugTrace(TargetingHandle, InTraceDatum.Start, InTraceDatum.End, InTraceDatum.Rot, bHasBlockingHit, InTraceDatum.OutHits);
+
 #endif // ENABLE_DRAW_DEBUG
 
 		ProcessHitResults(TargetingHandle, InTraceDatum.OutHits);
@@ -386,7 +405,6 @@ bool UTargetingSelectionTask_Trace::CanEditChange(const FProperty* InProperty) c
 #endif // WITH_EDITOR
 
 #if ENABLE_DRAW_DEBUG
-
 void UTargetingSelectionTask_Trace::DrawDebug(UTargetingSubsystem* TargetingSubsystem, FTargetingDebugInfo& Info, const FTargetingRequestHandle& TargetingHandle, float XOffset, float YOffset, int32 MinTextRowsToAdvance) const
 {
 #if WITH_EDITORONLY_DATA
@@ -408,6 +426,36 @@ void UTargetingSelectionTask_Trace::DrawDebug(UTargetingSubsystem* TargetingSubs
 #endif // WITH_EDITORONLY_DATA
 }
 
+void UTargetingSelectionTask_Trace::DrawDebugTrace(const FTargetingRequestHandle TargetingHandle, const FVector& StartLocation, const FVector& EndLocation, const FQuat& OrientationQuat, const bool bHit, const TArray<FHitResult>& Hits) const
+{
+	if (UTargetingSubsystem::IsTargetingDebugEnabled())
+	{
+		if (UWorld* World = GetSourceContextWorld(TargetingHandle))
+		{
+			const float DrawTime = UTargetingSubsystem::GetOverrideTargetingLifeTime();
+			const EDrawDebugTrace::Type DrawDebugType = DrawTime <= 0.0f ? EDrawDebugTrace::Type::ForOneFrame : EDrawDebugTrace::Type::ForDuration;
+			const FLinearColor TraceColor = FLinearColor::Red;
+			const FLinearColor TraceHitColor = FLinearColor::Green;
+			switch (TraceType)
+			{
+			case ETargetingTraceType::Sphere:
+				DrawDebugSphereTraceMulti(World, StartLocation, EndLocation, GetSweptTraceRadius(TargetingHandle), DrawDebugType, bHit, Hits, TraceColor, TraceHitColor, DrawTime);
+				break;
+			case ETargetingTraceType::Capsule:
+				DrawDebugCapsuleTraceMulti(World, StartLocation, EndLocation, GetSweptTraceRadius(TargetingHandle), GetSweptTraceCapsuleHalfHeight(TargetingHandle), OrientationQuat.Rotator(), DrawDebugType, bHit, Hits, TraceColor, TraceHitColor, DrawTime);
+				break;
+			case ETargetingTraceType::Box:
+				DrawDebugBoxTraceMulti(World, StartLocation, EndLocation, GetSweptTraceBoxHalfExtents(TargetingHandle), OrientationQuat.Rotator(), DrawDebugType, bHit, Hits, TraceColor, TraceHitColor, DrawTime);
+				break;
+			default:
+			case ETargetingTraceType::Line:
+				DrawDebugLineTraceMulti(World, StartLocation, EndLocation, DrawDebugType, bHit, Hits, TraceColor, TraceHitColor, DrawTime);
+				break;
+			}
+		}
+	}
+}
+
 void UTargetingSelectionTask_Trace::BuildTraceResultsDebugString(const FTargetingRequestHandle& TargetingHandle, const TArray<FTargetingDefaultResultData>& TargetResults) const
 {
 #if WITH_EDITORONLY_DATA
diff --git a/Engine/Plugins/Experimental/GameplayTargetingSystem/Source/GameplayTargetingSystem/Public/Tasks/TargetingSelectionTask_Trace.h b/Engine/Plugins/Experimental/GameplayTargetingSystem/Source/GameplayTargetingSystem/Public/Tasks/TargetingSelectionTask_Trace.h
--- a/Engine/Plugins/Experimental/GameplayTargetingSystem/Source/GameplayTargetingSystem/Public/Tasks/TargetingSelectionTask_Trace.h
+++ b/Engine/Plugins/Experimental/GameplayTargetingSystem/Source/GameplayTargetingSystem/Public/Tasks/TargetingSelectionTask_Trace.h
@@ -26,7 +26,7 @@ enum class ETargetingTraceType : uint8
 {
 	Line,
 	Sphere,
-	/* With zero rotation, the capsule length is along the cross product of the trace direction and Up. */
+	/* With zero rotation, the capsule length runs along the trace's up vector (the cross product of the trace's forward vector and the trace's forward vector flattened onto the XY plane and rotated 90 degrees around the Z vector). */
 	Capsule,
 	Box,
 };
@@ -175,6 +175,9 @@ protected:
 #if ENABLE_DRAW_DEBUG
 private:
 	virtual void DrawDebug(UTargetingSubsystem* TargetingSubsystem, FTargetingDebugInfo& Info, const FTargetingRequestHandle& TargetingHandle, float XOffset, float YOffset, int32 MinTextRowsToAdvance) const override;
+	
+	/** Draw debug info showing the results of the shape trace used for targeting. */
+	virtual void DrawDebugTrace(const FTargetingRequestHandle TargetingHandle, const FVector& StartLocation, const FVector& EndLocation, const FQuat& OrientationQuat, const bool bHit, const TArray<FHitResult>& Hits) const;
 	void BuildTraceResultsDebugString(const FTargetingRequestHandle& TargetingHandle, const TArray<FTargetingDefaultResultData>& TargetResults) const;
 	void ResetTraceResultsDebugString(const FTargetingRequestHandle& TargetingHandle) const;
 #endif // ENABLE_DRAW_DEBUG
diff --git a/Engine/Source/Runtime/Engine/Private/KismetTraceUtils.cpp b/Engine/Source/Runtime/Engine/Private/KismetTraceUtils.cpp
--- a/Engine/Source/Runtime/Engine/Private/KismetTraceUtils.cpp
+++ b/Engine/Source/Runtime/Engine/Private/KismetTraceUtils.cpp
@@ -271,56 +271,68 @@ void DrawDebugSphereTraceMulti(const UWorld* World, const FVector& Start, const
 }
 
 void DrawDebugCapsuleTraceSingle(const UWorld* World, const FVector& Start, const FVector& End, float Radius, float HalfHeight, EDrawDebugTrace::Type DrawDebugType, bool bHit, const FHitResult& OutHit, FLinearColor TraceColor, FLinearColor TraceHitColor, float DrawTime)
+{
+	DrawDebugCapsuleTraceSingle(World, Start, End, Radius, HalfHeight, FRotator::ZeroRotator, DrawDebugType, bHit, OutHit, TraceColor, TraceHitColor, DrawTime);
+}
+
+void DrawDebugCapsuleTraceSingle(const UWorld* World, const FVector& Start, const FVector& End, float Radius, float HalfHeight, const FRotator& Orientation, EDrawDebugTrace::Type DrawDebugType, bool bHit, const FHitResult& OutHit, FLinearColor TraceColor, FLinearColor TraceHitColor, float DrawTime)
 {
 	if (DrawDebugType != EDrawDebugTrace::None)
 	{
-		bool bPersistent = DrawDebugType == EDrawDebugTrace::Persistent;
-		float LifeTime = (DrawDebugType == EDrawDebugTrace::ForDuration) ? DrawTime : 0.f;
+		const bool bPersistent = DrawDebugType == EDrawDebugTrace::Persistent;
+		const float LifeTime = (DrawDebugType == EDrawDebugTrace::ForDuration) ? DrawTime : 0.f;
+		const FQuat OrientationQuat = Orientation.Quaternion();
 
 		if (bHit && OutHit.bBlockingHit)
 		{
 			// Red up to the blocking hit, green thereafter
-			::DrawDebugCapsule(World, Start, HalfHeight, Radius, FQuat::Identity, TraceColor.ToFColor(true), bPersistent, LifeTime);
-			::DrawDebugCapsule(World, OutHit.Location, HalfHeight, Radius, FQuat::Identity, TraceColor.ToFColor(true), bPersistent, LifeTime);
+			::DrawDebugCapsule(World, Start, HalfHeight, Radius, OrientationQuat, TraceColor.ToFColor(true), bPersistent, LifeTime);
+			::DrawDebugCapsule(World, OutHit.Location, HalfHeight, Radius, OrientationQuat, TraceColor.ToFColor(true), bPersistent, LifeTime);
 			::DrawDebugLine(World, Start, OutHit.Location, TraceColor.ToFColor(true), bPersistent, LifeTime);
 			::DrawDebugPoint(World, OutHit.ImpactPoint, KISMET_TRACE_DEBUG_IMPACTPOINT_SIZE, TraceColor.ToFColor(true), bPersistent, LifeTime);
 
-			::DrawDebugCapsule(World, End, HalfHeight, Radius, FQuat::Identity, TraceHitColor.ToFColor(true), bPersistent, LifeTime);
+			::DrawDebugCapsule(World, End, HalfHeight, Radius, OrientationQuat, TraceHitColor.ToFColor(true), bPersistent, LifeTime);
 			::DrawDebugLine(World, OutHit.Location, End, TraceHitColor.ToFColor(true), bPersistent, LifeTime);
 		}
 		else
 		{
 			// no hit means all red
-			::DrawDebugCapsule(World, Start, HalfHeight, Radius, FQuat::Identity, TraceColor.ToFColor(true), bPersistent, LifeTime);
-			::DrawDebugCapsule(World, End, HalfHeight, Radius, FQuat::Identity, TraceColor.ToFColor(true), bPersistent, LifeTime);
+			::DrawDebugCapsule(World, Start, HalfHeight, Radius, OrientationQuat, TraceColor.ToFColor(true), bPersistent, LifeTime);
+			::DrawDebugCapsule(World, End, HalfHeight, Radius, OrientationQuat, TraceColor.ToFColor(true), bPersistent, LifeTime);
 			::DrawDebugLine(World, Start, End, TraceColor.ToFColor(true), bPersistent, LifeTime);
 		}
 	}
 }
 
 void DrawDebugCapsuleTraceMulti(const UWorld* World, const FVector& Start, const FVector& End, float Radius, float HalfHeight, EDrawDebugTrace::Type DrawDebugType, bool bHit, const TArray<FHitResult>& OutHits, FLinearColor TraceColor, FLinearColor TraceHitColor, float DrawTime)
+{
+	DrawDebugCapsuleTraceMulti(World, Start, End, Radius, HalfHeight, FRotator::ZeroRotator, DrawDebugType, bHit, OutHits, TraceColor, TraceHitColor, DrawTime);
+}
+
+void DrawDebugCapsuleTraceMulti(const UWorld* World, const FVector& Start, const FVector& End, float Radius, float HalfHeight, const FRotator& Orientation, EDrawDebugTrace::Type DrawDebugType, bool bHit, const TArray<FHitResult>& OutHits, FLinearColor TraceColor, FLinearColor TraceHitColor, float DrawTime)
 {
 	if (DrawDebugType != EDrawDebugTrace::None)
 	{
-		bool bPersistent = DrawDebugType == EDrawDebugTrace::Persistent;
-		float LifeTime = (DrawDebugType == EDrawDebugTrace::ForDuration) ? DrawTime : 0.f;
+		const bool bPersistent = DrawDebugType == EDrawDebugTrace::Persistent;
+		const float LifeTime = (DrawDebugType == EDrawDebugTrace::ForDuration) ? DrawTime : 0.f;
+		const FQuat OrientationQuat = Orientation.Quaternion();
 
 		if (bHit && OutHits.Last().bBlockingHit)
 		{
 			// Red up to the blocking hit, green thereafter
 			FVector const BlockingHitPoint = OutHits.Last().Location;
-			::DrawDebugCapsule(World, Start, HalfHeight, Radius, FQuat::Identity, TraceColor.ToFColor(true), bPersistent, LifeTime);
-			::DrawDebugCapsule(World, BlockingHitPoint, HalfHeight, Radius, FQuat::Identity, TraceColor.ToFColor(true), bPersistent, LifeTime);
+			::DrawDebugCapsule(World, Start, HalfHeight, Radius, OrientationQuat, TraceColor.ToFColor(true), bPersistent, LifeTime);
+			::DrawDebugCapsule(World, BlockingHitPoint, HalfHeight, Radius, OrientationQuat, TraceColor.ToFColor(true), bPersistent, LifeTime);
 			::DrawDebugLine(World, Start, BlockingHitPoint, TraceColor.ToFColor(true), bPersistent, LifeTime);
 
-			::DrawDebugCapsule(World, End, HalfHeight, Radius, FQuat::Identity, TraceHitColor.ToFColor(true), bPersistent, LifeTime);
+			::DrawDebugCapsule(World, End, HalfHeight, Radius, OrientationQuat, TraceHitColor.ToFColor(true), bPersistent, LifeTime);
 			::DrawDebugLine(World, BlockingHitPoint, End, TraceHitColor.ToFColor(true), bPersistent, LifeTime);
 		}
 		else
 		{
 			// no hit means all red
-			::DrawDebugCapsule(World, Start, HalfHeight, Radius, FQuat::Identity, TraceColor.ToFColor(true), bPersistent, LifeTime);
-			::DrawDebugCapsule(World, End, HalfHeight, Radius, FQuat::Identity, TraceColor.ToFColor(true), bPersistent, LifeTime);
+			::DrawDebugCapsule(World, Start, HalfHeight, Radius, OrientationQuat, TraceColor.ToFColor(true), bPersistent, LifeTime);
+			::DrawDebugCapsule(World, End, HalfHeight, Radius, OrientationQuat, TraceColor.ToFColor(true), bPersistent, LifeTime);
 			::DrawDebugLine(World, Start, End, TraceColor.ToFColor(true), bPersistent, LifeTime);
 		}
 
diff --git a/Engine/Source/Runtime/Engine/Public/KismetTraceUtils.h b/Engine/Source/Runtime/Engine/Public/KismetTraceUtils.h
--- a/Engine/Source/Runtime/Engine/Public/KismetTraceUtils.h
+++ b/Engine/Source/Runtime/Engine/Public/KismetTraceUtils.h
@@ -27,9 +27,16 @@ ENGINE_API void DrawDebugBoxTraceMulti(const UWorld* World, const FVector& Start
 ENGINE_API void DrawDebugSphereTraceSingle(const UWorld* World, const FVector& Start, const FVector& End, float Radius, EDrawDebugTrace::Type DrawDebugType, bool bHit, const FHitResult& OutHit, FLinearColor TraceColor, FLinearColor TraceHitColor, float DrawTime);
 ENGINE_API void DrawDebugSphereTraceMulti(const UWorld* World, const FVector& Start, const FVector& End, float Radius, EDrawDebugTrace::Type DrawDebugType, bool bHit, const TArray<FHitResult>& OutHits, FLinearColor TraceColor, FLinearColor TraceHitColor, float DrawTime);
 
+/** Debug Drawing for single trace capsules using rotation FRotator::ZeroRotator. */
 ENGINE_API void DrawDebugCapsuleTraceSingle(const UWorld* World, const FVector& Start, const FVector& End, float Radius, float HalfHeight, EDrawDebugTrace::Type DrawDebugType, bool bHit, const FHitResult& OutHit, FLinearColor TraceColor, FLinearColor TraceHitColor, float DrawTime);
+/** Debug Drawing for multi trace capsules using rotation FRotator::ZeroRotator. */
 ENGINE_API void DrawDebugCapsuleTraceMulti(const UWorld* World, const FVector& Start, const FVector& End, float Radius, float HalfHeight, EDrawDebugTrace::Type DrawDebugType, bool bHit, const TArray<FHitResult>& OutHits, FLinearColor TraceColor, FLinearColor TraceHitColor, float DrawTime);
 
+/** Debug Drawing for single trace capsules with orientation. */
+ENGINE_API void DrawDebugCapsuleTraceSingle(const UWorld* World, const FVector& Start, const FVector& End, float Radius, float HalfHeight, const FRotator& Orientation, EDrawDebugTrace::Type DrawDebugType, bool bHit, const FHitResult& OutHit, FLinearColor TraceColor, FLinearColor TraceHitColor, float DrawTime);
+/** Debug Drawing for multi trace capsules with orientation. */
+ENGINE_API void DrawDebugCapsuleTraceMulti(const UWorld* World, const FVector& Start, const FVector& End, float Radius, float HalfHeight, const FRotator& Orientation, EDrawDebugTrace::Type DrawDebugType, bool bHit, const TArray<FHitResult>& OutHits, FLinearColor TraceColor, FLinearColor TraceHitColor, float DrawTime);
+
 #endif // ENABLE_DRAW_DEBUG
 
 
