GitHub 11226 : Fix newline escaping in DEFINE_TARGETING_DATA_STORE

A backslash shouldn't appear on the last line of a macro definition. The continuation of the newline can potentially yield compilation errors on some platforms.

diff --git a/Engine/Plugins/Experimental/GameplayTargetingSystem/Source/GameplayTargetingSystem/Public/Types/TargetingSystemDataStores.h b/Engine/Plugins/Experimental/GameplayTargetingSystem/Source/GameplayTargetingSystem/Public/Types/TargetingSystemDataStores.h
--- a/Engine/Plugins/Experimental/GameplayTargetingSystem/Source/GameplayTargetingSystem/Public/Types/TargetingSystemDataStores.h
+++ b/Engine/Plugins/Experimental/GameplayTargetingSystem/Source/GameplayTargetingSystem/Public/Types/TargetingSystemDataStores.h
@@ -237,36 +237,36 @@ namespace UE
  * This new Data Store could be accessed in a relevant TargetingTask like
  * TTargetingDataStore<FCustomTargetingData>::FindOrAdd(Handle)
  */ 
-#define DECLARE_TARGETING_DATA_STORE(DataType) namespace UE\
-{\
-	namespace TargetingSystem\
-	{\
-		extern TTargetingDataStore<DataType> GTargetingDataStore##DataType;\
-\
-		template<>\
-		FORCEINLINE DataType& TTargetingDataStore<DataType>::FindOrAdd(FTargetingRequestHandle Handle)\
-		{\
-			return GTargetingDataStore##DataType.Items.FindOrAdd(Handle);\
-		}\
-\
-		template<>\
-		FORCEINLINE DataType* TTargetingDataStore<DataType>::Find(FTargetingRequestHandle Handle)\
-		{\
-			return GTargetingDataStore##DataType.Items.Find(Handle);\
-		}\
-\
-		template<>\
-		FORCEINLINE void TTargetingDataStore<DataType>::OnTargetingRequestHandleReleased(FTargetingRequestHandle Handle)\
-		{\
-			GTargetingDataStore##DataType.Items.Remove(Handle);\
-		}\
-	}\
+#define DECLARE_TARGETING_DATA_STORE(DataType) namespace UE																	\
+{																															\
+	namespace TargetingSystem																								\
+	{																														\
+		extern TTargetingDataStore<DataType> GTargetingDataStore##DataType;													\
+																															\
+		template<>																											\
+		FORCEINLINE DataType& TTargetingDataStore<DataType>::FindOrAdd(FTargetingRequestHandle Handle)						\
+		{																													\
+			return GTargetingDataStore##DataType.Items.FindOrAdd(Handle);													\
+		}																													\
+																															\
+		template<>																											\
+		FORCEINLINE DataType* TTargetingDataStore<DataType>::Find(FTargetingRequestHandle Handle)							\
+		{																													\
+			return GTargetingDataStore##DataType.Items.Find(Handle);														\
+		}																													\
+																															\
+		template<>																											\
+		FORCEINLINE void TTargetingDataStore<DataType>::OnTargetingRequestHandleReleased(FTargetingRequestHandle Handle)	\
+		{																													\
+			GTargetingDataStore##DataType.Items.Remove(Handle);																\
+		}																													\
+	}																														\
 }
 
-#define DEFINE_TARGETING_DATA_STORE(DataType) namespace UE\
-{\
-	namespace TargetingSystem\
-	{\
-		TTargetingDataStore<DataType> GTargetingDataStore##DataType;\
-	}\
-}\
+#define DEFINE_TARGETING_DATA_STORE(DataType) namespace UE																	\
+{																															\
+	namespace TargetingSystem																								\
+	{																														\
+		TTargetingDataStore<DataType> GTargetingDataStore##DataType;														\
+	}																														\
+}
