[Controls] Targeting Preset Interaction Selection
#rnx
#tests [How To Test]
* Place the following devices on your island
   * Control: Third Person
   * Camera: Fixed Angle
* Play the game
* Use this command to ease the testing ts.debug.EnableTargetingDebugging true
* you will be able to see we're interacting with the object depending on the config we have on the targeting preset + the interact values.
* You can modify the Interact value directly from the device, but the highlight interact value is actually hardcoded in blueprints, so depending on the interact rate you will be able to see how you can see the object highlighted but not able to interact with it if you are too far

diff --git a/Engine/Plugins/Experimental/GameplayTargetingSystem/Source/GameplayTargetingSystem/Private/Tasks/SimpleTargetingSelectionTask.cpp b/Engine/Plugins/Experimental/GameplayTargetingSystem/Source/GameplayTargetingSystem/Private/Tasks/SimpleTargetingSelectionTask.cpp
--- a/Engine/Plugins/Experimental/GameplayTargetingSystem/Source/GameplayTargetingSystem/Private/Tasks/SimpleTargetingSelectionTask.cpp
+++ b/Engine/Plugins/Experimental/GameplayTargetingSystem/Source/GameplayTargetingSystem/Private/Tasks/SimpleTargetingSelectionTask.cpp
@@ -27,6 +27,13 @@ bool USimpleTargetingSelectionTask::AddTargetActor(const FTargetingRequestHandle
 			FTargetingDefaultResultData& ResultData = ResultsSet.TargetResults.AddDefaulted_GetRef();
 			ResultData.HitResult.HitObjectHandle = FActorInstanceHandle(Actor);
 			ResultData.HitResult.Location = Actor->GetActorLocation();
+			if (const FTargetingSourceContext* SourceContext = FTargetingSourceContext::Find(TargetingHandle))
+			{
+				if (SourceContext->SourceActor)
+				{
+					ResultData.HitResult.Distance = FVector::Distance(SourceContext->SourceActor->GetActorLocation(), Actor->GetActorLocation());
+				}
+			}
 			return true;
 		}
 	}
@@ -40,6 +47,13 @@ bool USimpleTargetingSelectionTask::AddHitResult(const FTargetingRequestHandle&
 	{
 		FTargetingDefaultResultData& ResultData = ResultsSet.TargetResults.AddDefaulted_GetRef();
 		ResultData.HitResult = HitResult;
+		if (const FTargetingSourceContext* SourceContext = FTargetingSourceContext::Find(TargetingHandle))
+		{
+			if (SourceContext->SourceActor)
+			{
+				ResultData.HitResult.Distance = FVector::Distance(SourceContext->SourceActor->GetActorLocation(), HitResult.GetActor()->GetActorLocation());
+			}
+		}
 		return true;
 	}
 	return false;
diff --git a/Engine/Plugins/Experimental/GameplayTargetingSystem/Source/GameplayTargetingSystem/Private/Tasks/TargetingSelectionTask_AOE.cpp b/Engine/Plugins/Experimental/GameplayTargetingSystem/Source/GameplayTargetingSystem/Private/Tasks/TargetingSelectionTask_AOE.cpp
--- a/Engine/Plugins/Experimental/GameplayTargetingSystem/Source/GameplayTargetingSystem/Private/Tasks/TargetingSelectionTask_AOE.cpp
+++ b/Engine/Plugins/Experimental/GameplayTargetingSystem/Source/GameplayTargetingSystem/Private/Tasks/TargetingSelectionTask_AOE.cpp
@@ -268,6 +268,7 @@ void UTargetingSelectionTask_AOE::ProcessOverlapResults(const FTargetingRequestH
 				ResultData->HitResult.Component = OverlapResult.GetComponent();
 				ResultData->HitResult.bBlockingHit = OverlapResult.bBlockingHit;
 				ResultData->HitResult.TraceStart = SourceLocation;
+				ResultData->HitResult.Distance = FVector::Distance(OverlapResult.GetActor()->GetActorLocation(), SourceLocation);
 			}
 		}
 
diff --git a/Engine/Plugins/Experimental/GameplayTargetingSystem/Source/GameplayTargetingSystem/Private/Tasks/TargetingSelectionTask_SourceActor.cpp b/Engine/Plugins/Experimental/GameplayTargetingSystem/Source/GameplayTargetingSystem/Private/Tasks/TargetingSelectionTask_SourceActor.cpp
--- a/Engine/Plugins/Experimental/GameplayTargetingSystem/Source/GameplayTargetingSystem/Private/Tasks/TargetingSelectionTask_SourceActor.cpp
+++ b/Engine/Plugins/Experimental/GameplayTargetingSystem/Source/GameplayTargetingSystem/Private/Tasks/TargetingSelectionTask_SourceActor.cpp
@@ -47,6 +47,8 @@ void UTargetingSelectionTask_SourceActor::Execute(const FTargetingRequestHandle&
 					FTargetingDefaultResultData* ResultData = new(Results.TargetResults) FTargetingDefaultResultData();
 					ResultData->HitResult.HitObjectHandle = FActorInstanceHandle(SourceContext->SourceActor);
 					ResultData->HitResult.Location = SourceContext->SourceActor->GetActorLocation();
+					ResultData->HitResult.Distance = 0.0f;
+					ResultData->HitResult.Time = 0.f;
 				}
 
 #if ENABLE_DRAW_DEBUG
