[Controls] AOE Targeting Preset not updating ImpactPoint for interact selection
#rnx
#tests [How To Test]
* On your island add the following devices
   * Control: Third Person
   * Vehicle devices
   * Buttons
   * Other interactable devices
* Play the game.
* Get near those objects and check you're able to interact with them
[Video after fix]
https://drive.google.com/file/d/1gdlbdVy7uG1ivOQ-TwSLAPFNN2YOi_6h/view?usp=drive_link

diff --git a/Engine/Plugins/Experimental/GameplayTargetingSystem/Source/GameplayTargetingSystem/Private/Tasks/TargetingSelectionTask_AOE.cpp b/Engine/Plugins/Experimental/GameplayTargetingSystem/Source/GameplayTargetingSystem/Private/Tasks/TargetingSelectionTask_AOE.cpp
--- a/Engine/Plugins/Experimental/GameplayTargetingSystem/Source/GameplayTargetingSystem/Private/Tasks/TargetingSelectionTask_AOE.cpp
+++ b/Engine/Plugins/Experimental/GameplayTargetingSystem/Source/GameplayTargetingSystem/Private/Tasks/TargetingSelectionTask_AOE.cpp
@@ -266,6 +266,7 @@ void UTargetingSelectionTask_AOE::ProcessOverlapResults(const FTargetingRequestH
 				FTargetingDefaultResultData* ResultData = new(TargetingResults.TargetResults) FTargetingDefaultResultData();
 				ResultData->HitResult.HitObjectHandle = FActorInstanceHandle(OverlapResult.GetActor(), OverlapResult.GetComponent(), OverlapResult.ItemIndex);
 				ResultData->HitResult.Component = OverlapResult.GetComponent();
+				ResultData->HitResult.ImpactPoint = OverlapResult.GetActor()->GetActorLocation();
 				ResultData->HitResult.bBlockingHit = OverlapResult.bBlockingHit;
 				ResultData->HitResult.TraceStart = SourceLocation;
 				ResultData->HitResult.Distance = FVector::Distance(OverlapResult.GetActor()->GetActorLocation(), SourceLocation);
