[Controls] Can't interact with doors
#rnx
#tests [How To Test]
* Place the following devices on your island
   * Control: Third Person
* Add different kind of doors.
* Play the game
* Interact with doors, you will be able to interact with all kinds of doors.

#virtualized

diff --git a/Engine/Plugins/Experimental/GameplayTargetingSystem/Source/GameplayTargetingSystem/Private/Tasks/SimpleTargetingSelectionTask.cpp b/Engine/Plugins/Experimental/GameplayTargetingSystem/Source/GameplayTargetingSystem/Private/Tasks/SimpleTargetingSelectionTask.cpp
--- a/Engine/Plugins/Experimental/GameplayTargetingSystem/Source/GameplayTargetingSystem/Private/Tasks/SimpleTargetingSelectionTask.cpp
+++ b/Engine/Plugins/Experimental/GameplayTargetingSystem/Source/GameplayTargetingSystem/Private/Tasks/SimpleTargetingSelectionTask.cpp
@@ -43,7 +43,12 @@ bool USimpleTargetingSelectionTask::AddTargetActor(const FTargetingRequestHandle
 bool USimpleTargetingSelectionTask::AddHitResult(const FTargetingRequestHandle& TargetingHandle, const FHitResult& HitResult) const
 {
 	FTargetingDefaultResultsSet& ResultsSet = FTargetingDefaultResultsSet::FindOrAdd(TargetingHandle);
-	if (!ResultsSet.TargetResults.FindByPredicate([HitResult](const FTargetingDefaultResultData& ResultData){ return ResultData.HitResult.GetActor() == HitResult.GetActor(); }))
+	if (!ResultsSet.TargetResults.FindByPredicate([HitResult](const FTargetingDefaultResultData& ResultData)
+	{
+		bool bIsSameActor = ResultData.HitResult.GetActor() == HitResult.GetActor();
+		bool bHasSameComponent = ResultData.HitResult.GetComponent() == HitResult.GetComponent();
+		return bIsSameActor && bHasSameComponent;
+	}))
 	{
 		FTargetingDefaultResultData& ResultData = ResultsSet.TargetResults.AddDefaulted_GetRef();
 		ResultData.HitResult = HitResult;
