Targeting API: Adding UTargetingSubsystem::GetTargetingSourceContext() to the TargetingSystem API

#rnx

diff --git a/Engine/Plugins/Experimental/GameplayTargetingSystem/Source/GameplayTargetingSystem/Public/TargetingSystem/TargetingSubsystem.h b/Engine/Plugins/Experimental/GameplayTargetingSystem/Source/GameplayTargetingSystem/Public/TargetingSystem/TargetingSubsystem.h
--- a/Engine/Plugins/Experimental/GameplayTargetingSystem/Source/GameplayTargetingSystem/Public/TargetingSystem/TargetingSubsystem.h
+++ b/Engine/Plugins/Experimental/GameplayTargetingSystem/Source/GameplayTargetingSystem/Public/TargetingSystem/TargetingSubsystem.h
@@ -191,7 +191,7 @@ private:
 public:
 	/** Returns the targeting source context for the targeting request handle */
 	UFUNCTION(BlueprintCallable, Category = "Targeting System | Targeting Types")
-	FTargetingSourceContext GetTargetingSourceContext(FTargetingRequestHandle TargetingHandle) const;
+	TARGETINGSYSTEM_API FTargetingSourceContext GetTargetingSourceContext(FTargetingRequestHandle TargetingHandle) const;
 
 	/** Method to get the actor targets from a given targeting request handle */
 	UFUNCTION(BlueprintCallable, Category = "Targeting System | Targeting Results")
