Add location to HitResults from Targeting Selection Task AOE

diff --git a/Engine/Plugins/Experimental/GameplayTargetingSystem/Source/GameplayTargetingSystem/Private/Tasks/TargetingSelectionTask_AOE.cpp b/Engine/Plugins/Experimental/GameplayTargetingSystem/Source/GameplayTargetingSystem/Private/Tasks/TargetingSelectionTask_AOE.cpp
--- a/Engine/Plugins/Experimental/GameplayTargetingSystem/Source/GameplayTargetingSystem/Private/Tasks/TargetingSelectionTask_AOE.cpp
+++ b/Engine/Plugins/Experimental/GameplayTargetingSystem/Source/GameplayTargetingSystem/Private/Tasks/TargetingSelectionTask_AOE.cpp
@@ -267,6 +267,7 @@ void UTargetingSelectionTask_AOE::ProcessOverlapResults(const FTargetingRequestH
 				ResultData->HitResult.HitObjectHandle = FActorInstanceHandle(OverlapResult.GetActor(), OverlapResult.GetComponent(), OverlapResult.ItemIndex);
 				ResultData->HitResult.Component = OverlapResult.GetComponent();
 				ResultData->HitResult.ImpactPoint = OverlapResult.GetActor()->GetActorLocation();
+				ResultData->HitResult.Location = OverlapResult.GetActor()->GetActorLocation();
 				ResultData->HitResult.bBlockingHit = OverlapResult.bBlockingHit;
 				ResultData->HitResult.TraceStart = SourceLocation;
 				ResultData->HitResult.Distance = FVector::Distance(OverlapResult.GetActor()->GetActorLocation(), SourceLocation);
