Fix ActorInstanceHandle edge-cases

diff --git a/Engine/Plugins/Experimental/GameplayTargetingSystem/Source/GameplayTargetingSystem/Private/Tasks/TargetingSelectionTask_AOE.cpp b/Engine/Plugins/Experimental/GameplayTargetingSystem/Source/GameplayTargetingSystem/Private/Tasks/TargetingSelectionTask_AOE.cpp
--- a/Engine/Plugins/Experimental/GameplayTargetingSystem/Source/GameplayTargetingSystem/Private/Tasks/TargetingSelectionTask_AOE.cpp
+++ b/Engine/Plugins/Experimental/GameplayTargetingSystem/Source/GameplayTargetingSystem/Private/Tasks/TargetingSelectionTask_AOE.cpp
@@ -264,12 +264,13 @@ void UTargetingSelectionTask_AOE::ProcessOverlapResults(const FTargetingRequestH
 			if (bAddResult)
 			{
 				FTargetingDefaultResultData* ResultData = new(TargetingResults.TargetResults) FTargetingDefaultResultData();
-				ResultData->HitResult.HitObjectHandle = FActorInstanceHandle(OverlapResult.GetActor(), OverlapResult.GetComponent(), OverlapResult.ItemIndex);
+				ResultData->HitResult.HitObjectHandle = OverlapResult.OverlapObjectHandle;
 				ResultData->HitResult.Component = OverlapResult.GetComponent();
 				ResultData->HitResult.ImpactPoint = OverlapResult.GetActor()->GetActorLocation();
 				ResultData->HitResult.Location = OverlapResult.GetActor()->GetActorLocation();
 				ResultData->HitResult.bBlockingHit = OverlapResult.bBlockingHit;
 				ResultData->HitResult.TraceStart = SourceLocation;
+				ResultData->HitResult.Item = OverlapResult.ItemIndex;
 				ResultData->HitResult.Distance = FVector::Distance(OverlapResult.GetActor()->GetActorLocation(), SourceLocation);
 			}
 		}
diff --git a/Engine/Source/Runtime/Engine/Private/Engine/ActorInstanceHandle.cpp b/Engine/Source/Runtime/Engine/Private/Engine/ActorInstanceHandle.cpp
--- a/Engine/Source/Runtime/Engine/Private/Engine/ActorInstanceHandle.cpp
+++ b/Engine/Source/Runtime/Engine/Private/Engine/ActorInstanceHandle.cpp
@@ -461,12 +461,13 @@ FArchive& operator<<(FArchive& Ar, FActorInstanceHandle& Handle)
 	Ar << Handle.ReferenceObject;
 	Ar << Handle.InstanceIndex;
 
-	TSoftObjectPtr<UObject> SoftObject(Handle.ManagerInterface.GetObject());
-	Ar << SoftObject;
+	// Serializing as a TWeakObjectPtr instead of TSoftObjectPtr to avoid UE-205038.
+	TWeakObjectPtr<UObject> WeakManagerObject(Handle.ManagerInterface.GetObject());
+	Ar << WeakManagerObject;
 
 	if (Ar.IsLoading())
 	{
-		Handle.ManagerInterface = FActorInstanceManagerInterface(SoftObject.Get());
+		Handle.ManagerInterface = FActorInstanceManagerInterface(WeakManagerObject.Get());
 	}
 
 	return Ar;
