Targeting System / Sort By Distance: Added option for more accurate distance calculation based off of nearest collision surface instead of actor location.

diff --git a/Engine/Plugins/Experimental/GameplayTargetingSystem/Source/GameplayTargetingSystem/Private/Tasks/TargetingFilterTask_SortByDistance.cpp b/Engine/Plugins/Experimental/GameplayTargetingSystem/Source/GameplayTargetingSystem/Private/Tasks/TargetingFilterTask_SortByDistance.cpp
--- a/Engine/Plugins/Experimental/GameplayTargetingSystem/Source/GameplayTargetingSystem/Private/Tasks/TargetingFilterTask_SortByDistance.cpp
+++ b/Engine/Plugins/Experimental/GameplayTargetingSystem/Source/GameplayTargetingSystem/Private/Tasks/TargetingFilterTask_SortByDistance.cpp
@@ -11,7 +11,20 @@ float UTargetingFilterTask_SortByDistance::GetScoreForTarget(const FTargetingReq
 	FVector SourceLocation = GetSourceLocation(TargetingHandle);
 	if (AActor* TargetActor = TargetData.HitResult.GetActor())
 	{
-		FVector TargetLocation = TargetActor->GetActorLocation();
+		FVector TargetLocation;
+		
+		if (bUseDistanceToNearestBlockingCollider)
+		{
+			const float DistanceToCollision = TargetActor->ActorGetDistanceToCollision(SourceLocation, DistanceToCollisionChannel, TargetLocation);
+
+			if (DistanceToCollision >= 0.f)
+			{
+				return DistanceToCollision * DistanceToCollision;
+			}
+		}
+
+		TargetLocation = TargetActor->GetActorLocation();
+
 		return FVector::DistSquared(SourceLocation, TargetLocation);
 	}
 	return 0.f;
diff --git a/Engine/Plugins/Experimental/GameplayTargetingSystem/Source/GameplayTargetingSystem/Public/Tasks/TargetingFilterTask_SortByDistance.h b/Engine/Plugins/Experimental/GameplayTargetingSystem/Source/GameplayTargetingSystem/Public/Tasks/TargetingFilterTask_SortByDistance.h
--- a/Engine/Plugins/Experimental/GameplayTargetingSystem/Source/GameplayTargetingSystem/Public/Tasks/TargetingFilterTask_SortByDistance.h
+++ b/Engine/Plugins/Experimental/GameplayTargetingSystem/Source/GameplayTargetingSystem/Public/Tasks/TargetingFilterTask_SortByDistance.h
@@ -30,6 +30,17 @@ class UTargetingFilterTask_SortByDistance : public UTargetingSortTask_Base
 protected:
 	virtual float GetScoreForTarget(const FTargetingRequestHandle& TargetingHandle, const FTargetingDefaultResultData& TargetData) const override;
 
+	/** Use the distance to the nearest blocking collision surface on each actor instead of the distance to the actor's location */
+	UPROPERTY(EditAnywhere, Category = "Target Sorting | Sort By Distance Data")
+	bool bUseDistanceToNearestBlockingCollider = false;
+
+	/** 
+	* Collision channel used to determine the closest point on a actor's collider. 
+	* Colliders will only be considered if they block this channel
+	*/
+	UPROPERTY(EditAnywhere, Category = "Target Sorting | Sort By Distance Data", Meta=(EditCondition="bUseDistanceToNearestBlockingCollider "))
+	TEnumAsByte<ECollisionChannel> DistanceToCollisionChannel = ECollisionChannel::ECC_WorldDynamic;
+
 private:
 	/** Returns the Source location, whether that comes from the source actor or the Source location in the Source Context */
 	FVector GetSourceLocation(const FTargetingRequestHandle& TargetingHandle) const;
