Moving around some things in the TargetingSystem to fix an issue with circular includes

PR #11369: Targeting System - Include fix for Platforms and Non Unity.

diff --git a/Engine/Plugins/Experimental/GameplayTargetingSystem/Source/GameplayTargetingSystem/Private/TargetingSystem/TargetingSubsystem.cpp b/Engine/Plugins/Experimental/GameplayTargetingSystem/Source/GameplayTargetingSystem/Private/TargetingSystem/TargetingSubsystem.cpp
--- a/Engine/Plugins/Experimental/GameplayTargetingSystem/Source/GameplayTargetingSystem/Private/TargetingSystem/TargetingSubsystem.cpp
+++ b/Engine/Plugins/Experimental/GameplayTargetingSystem/Source/GameplayTargetingSystem/Private/TargetingSystem/TargetingSubsystem.cpp
@@ -12,6 +12,7 @@
 #include "Tasks/CollisionQueryTaskData.h"
 #include "Types/TargetingSystemLogs.h"
 #include "Types/TargetingSystemTypes.h"
+#include "Tasks/TargetingTask.h"
 
 #if ENABLE_DRAW_DEBUG
 #include "GameFramework/HUD.h"
@@ -75,10 +76,9 @@ namespace TargetingSystemCVars
 #endif // ENABLE_DRAW_DEBUG
 }
 
-UTargetingSubsystem::FOnTargetingRequestHandleReleased& UTargetingSubsystem::ReleaseHandleDelegate()
+FTargetingRequestHandle::FOnTargetingRequestHandleReleased& UTargetingSubsystem::ReleaseHandleDelegate()
 {
-	static UTargetingSubsystem::FOnTargetingRequestHandleReleased Delegate;
-	return Delegate;
+	return FTargetingRequestHandle::GetReleaseHandleDelegate();
 }
 
 
@@ -344,7 +344,7 @@ void UTargetingSubsystem::ReleaseTargetRequestHandle(FTargetingRequestHandle& Ha
 	}
 #endif // ENABLE_DRAW_DEBUG
 
-	ReleaseHandleDelegate().Broadcast(CachedHandle);
+	FTargetingRequestHandle::GetReleaseHandleDelegate().Broadcast(CachedHandle);
 	TARGETING_LOG(Verbose, TEXT("%s: - Releasigng Handle [%d]"), ANSI_TO_TCHAR(__FUNCTION__), CachedHandle.Handle);
 }
 
diff --git a/Engine/Plugins/Experimental/GameplayTargetingSystem/Source/GameplayTargetingSystem/Private/Types/TargetingSystemTypes.cpp b/Engine/Plugins/Experimental/GameplayTargetingSystem/Source/GameplayTargetingSystem/Private/Types/TargetingSystemTypes.cpp
--- a/Engine/Plugins/Experimental/GameplayTargetingSystem/Source/GameplayTargetingSystem/Private/Types/TargetingSystemTypes.cpp
+++ b/Engine/Plugins/Experimental/GameplayTargetingSystem/Source/GameplayTargetingSystem/Private/Types/TargetingSystemTypes.cpp
@@ -21,6 +21,12 @@ bool FTargetingRequestHandle::NetSerialize(FArchive& Ar, class UPackageMap* Map,
 	return false;
 }
 
+FTargetingRequestHandle::FOnTargetingRequestHandleReleased& FTargetingRequestHandle::GetReleaseHandleDelegate()
+{
+	static FTargetingRequestHandle::FOnTargetingRequestHandleReleased Delegate;
+	return Delegate;
+}
+
 /**	@struct FTargetingTaskSet */
 
 const FTargetingTaskSet*& FTargetingTaskSet::FindOrAdd(FTargetingRequestHandle Handle)
diff --git a/Engine/Plugins/Experimental/GameplayTargetingSystem/Source/GameplayTargetingSystem/Public/TargetingSystem/TargetingSubsystem.h b/Engine/Plugins/Experimental/GameplayTargetingSystem/Source/GameplayTargetingSystem/Public/TargetingSystem/TargetingSubsystem.h
--- a/Engine/Plugins/Experimental/GameplayTargetingSystem/Source/GameplayTargetingSystem/Public/TargetingSystem/TargetingSubsystem.h
+++ b/Engine/Plugins/Experimental/GameplayTargetingSystem/Source/GameplayTargetingSystem/Public/TargetingSystem/TargetingSubsystem.h
@@ -6,7 +6,7 @@
 #include "DrawDebugHelpers.h"
 #include "Misc/CoreMisc.h"
 #include "Subsystems/GameInstanceSubsystem.h"
-#include "Tasks/TargetingTask.h"
+#include "Tasks/CollisionQueryTaskData.h" 
 #include "Tickable.h"
 #include "Types/TargetingSystemTypes.h"
 
@@ -19,7 +19,6 @@ class UCanvas;
 class UTargetingPreset;
 class UTargetingTask;
 class UWorld;
-struct FCollisionQueryTaskData;
 
 
 #if ENABLE_DRAW_DEBUG
@@ -124,8 +123,13 @@ public:
 	TARGETINGSYSTEM_API static void ReleaseTargetRequestHandle(FTargetingRequestHandle& Handle);
 
 	/** The handle released delegate that fires right before a handle is release so all data stores can clean up their state */
+	UE_DEPRECATED(5.4, "Use FTargetingRequestHandle::FOnTargetingRequestHandleReleased instead")
 	DECLARE_MULTICAST_DELEGATE_OneParam(FOnTargetingRequestHandleReleased, FTargetingRequestHandle)
-	TARGETINGSYSTEM_API static FOnTargetingRequestHandleReleased& ReleaseHandleDelegate();
+
+	UE_DEPRECATED(5.4, "Call FTargetingRequestHandle::GetReleaseHandleDelegate instead.")
+PRAGMA_DISABLE_DEPRECATION_WARNINGS
+	TARGETINGSYSTEM_API static FTargetingRequestHandle::FOnTargetingRequestHandleReleased& ReleaseHandleDelegate();
+PRAGMA_ENABLE_DEPRECATION_WARNINGS
 
 	/** Target Handle Generation Methods */
 
diff --git a/Engine/Plugins/Experimental/GameplayTargetingSystem/Source/GameplayTargetingSystem/Public/Types/TargetingSystemDataStores.h b/Engine/Plugins/Experimental/GameplayTargetingSystem/Source/GameplayTargetingSystem/Public/Types/TargetingSystemDataStores.h
--- a/Engine/Plugins/Experimental/GameplayTargetingSystem/Source/GameplayTargetingSystem/Public/Types/TargetingSystemDataStores.h
+++ b/Engine/Plugins/Experimental/GameplayTargetingSystem/Source/GameplayTargetingSystem/Public/Types/TargetingSystemDataStores.h
@@ -5,7 +5,6 @@
 #include "CoreMinimal.h"
 #include "Delegates/DelegateCombinations.h"
 #include "Engine/EngineTypes.h"
-#include "TargetingSystem/TargetingSubsystem.h"
 #include "TargetingSystemTypes.h"
 
 
@@ -28,14 +27,14 @@ namespace UE
 		{
 			TTargetingDataStore()
 			{
-				ReleaseDelegateHandle = UTargetingSubsystem::ReleaseHandleDelegate().AddStatic(&TTargetingDataStore::OnTargetingRequestHandleReleased);
+				ReleaseDelegateHandle = FTargetingRequestHandle::GetReleaseHandleDelegate().AddStatic(&TTargetingDataStore::OnTargetingRequestHandleReleased);
 			}
 
 			~TTargetingDataStore()
 			{
 				if (ReleaseDelegateHandle.IsValid())
 				{
-					UTargetingSubsystem::ReleaseHandleDelegate().Remove(ReleaseDelegateHandle);
+					FTargetingRequestHandle::GetReleaseHandleDelegate().Remove(ReleaseDelegateHandle);
 				}
 			}
 
diff --git a/Engine/Plugins/Experimental/GameplayTargetingSystem/Source/GameplayTargetingSystem/Public/Types/TargetingSystemTypes.h b/Engine/Plugins/Experimental/GameplayTargetingSystem/Source/GameplayTargetingSystem/Public/Types/TargetingSystemTypes.h
--- a/Engine/Plugins/Experimental/GameplayTargetingSystem/Source/GameplayTargetingSystem/Public/Types/TargetingSystemTypes.h
+++ b/Engine/Plugins/Experimental/GameplayTargetingSystem/Source/GameplayTargetingSystem/Public/Types/TargetingSystemTypes.h
@@ -50,6 +50,10 @@ public:
 	/** We override NetSeralize because we should never be network serializing this handle */
 	bool NetSerialize(FArchive& Ar, class UPackageMap* Map, bool& bOutSuccess);
 
+	/** The handle released delegate that fires right before a handle is release so all data stores can clean up their state */
+	DECLARE_MULTICAST_DELEGATE_OneParam(FOnTargetingRequestHandleReleased, FTargetingRequestHandle)
+	TARGETINGSYSTEM_API static FOnTargetingRequestHandleReleased& GetReleaseHandleDelegate();
+	
 	/** overloaded operators */
 	FORCEINLINE bool operator==(const FTargetingRequestHandle& InHandle) const { return Handle == InHandle.Handle; }
 	FORCEINLINE bool operator!=(const FTargetingRequestHandle& InHandle) const { return !this->operator==(InHandle); }
