Add option to trace complex geometry on AoE targeting task

diff --git a/Engine/Plugins/Experimental/GameplayTargetingSystem/Source/GameplayTargetingSystem/Private/Tasks/TargetingSelectionTask_AOE.cpp b/Engine/Plugins/Experimental/GameplayTargetingSystem/Source/GameplayTargetingSystem/Private/Tasks/TargetingSelectionTask_AOE.cpp
--- a/Engine/Plugins/Experimental/GameplayTargetingSystem/Source/GameplayTargetingSystem/Private/Tasks/TargetingSelectionTask_AOE.cpp
+++ b/Engine/Plugins/Experimental/GameplayTargetingSystem/Source/GameplayTargetingSystem/Private/Tasks/TargetingSelectionTask_AOE.cpp
@@ -17,6 +17,16 @@
 #include "Engine/Canvas.h"
 #endif // ENABLE_DRAW_DEBUG
 
+namespace TargetingSystemCVars
+{
+	static bool bEnableComplexTracingAOE = true;
+	FAutoConsoleVariableRef CVarEnableComplexTracingAOE(
+		TEXT("ts.AOE.EnableComplexTracingAOE"),
+		bEnableComplexTracingAOE,
+		TEXT("When enabled, allows users to do complex traces by setting bShouldTraceComplex to true.")
+	);
+}
+
 UTargetingSelectionTask_AOE::UTargetingSelectionTask_AOE(const FObjectInitializer& ObjectInitializer)
 : Super(ObjectInitializer)
 {
@@ -389,6 +399,8 @@ void UTargetingSelectionTask_AOE::InitCollisionParams(const FTargetingRequestHan
 	{
 		OutParams.AddIgnoredActors(FoundOverride->IgnoredActors);
 	}
+
+	OutParams.bTraceComplex = TargetingSystemCVars::bEnableComplexTracingAOE && bTraceComplex;
 }
 
 #if WITH_EDITOR
diff --git a/Engine/Plugins/Experimental/GameplayTargetingSystem/Source/GameplayTargetingSystem/Public/Tasks/TargetingSelectionTask_AOE.h b/Engine/Plugins/Experimental/GameplayTargetingSystem/Source/GameplayTargetingSystem/Public/Tasks/TargetingSelectionTask_AOE.h
--- a/Engine/Plugins/Experimental/GameplayTargetingSystem/Source/GameplayTargetingSystem/Public/Tasks/TargetingSelectionTask_AOE.h
+++ b/Engine/Plugins/Experimental/GameplayTargetingSystem/Source/GameplayTargetingSystem/Public/Tasks/TargetingSelectionTask_AOE.h
@@ -131,6 +131,10 @@ protected:
 	UPROPERTY(EditAnywhere, Category = "Target AOE Selection | Data")
 	uint8 bIgnoreInstigatorActor : 1;
 
+	/** When enabled, the trace will be performed against complex collision. */
+	UPROPERTY(EditAnywhere, Category = "Target AOE Selection | Data")
+	uint8 bTraceComplex : 1 = false;
+
 	/** The half extent to use for box and cylinder */
 	UPROPERTY(EditAnywhere, Category = "Target AOE Selection | Box and Cylinder")
 	FVector HalfExtent = FVector::ZeroVector;
