Changing SelectTargets into a BlueprintNativeEvent so we can implement it and override it from source

diff --git a/Engine/Plugins/Experimental/GameplayTargetingSystem/Source/GameplayTargetingSystem/Public/Tasks/SimpleTargetingSelectionTask.h b/Engine/Plugins/Experimental/GameplayTargetingSystem/Source/GameplayTargetingSystem/Public/Tasks/SimpleTargetingSelectionTask.h
--- a/Engine/Plugins/Experimental/GameplayTargetingSystem/Source/GameplayTargetingSystem/Public/Tasks/SimpleTargetingSelectionTask.h
+++ b/Engine/Plugins/Experimental/GameplayTargetingSystem/Source/GameplayTargetingSystem/Public/Tasks/SimpleTargetingSelectionTask.h
@@ -25,9 +25,11 @@ public:
 	/** Evaluation function called by derived classes to process the targeting request */
 	virtual void Execute(const FTargetingRequestHandle& TargetingHandle) const override;
 	
-	UFUNCTION(BlueprintImplementableEvent)
+	UFUNCTION(BlueprintNativeEvent)
 	void SelectTargets(const FTargetingRequestHandle& TargetingHandle, const FTargetingSourceContext& SourceContext) const;
 
+	virtual void SelectTargets_Implementation(const FTargetingRequestHandle& TargetingHandle, const FTargetingSourceContext& SourceContext) const {};
+
 protected:
 	/**
 	 * Adds a single Actor to the Targeting Results for a given TargetingRequestHandle.
