Add `bStableSort` to `UTargetingSortTask_Base`

With this property, any sort task can opt-in to a slightly slower sorting algorithm that preserves the relative ordering of targets with equal scores.

diff --git a/Engine/Plugins/Experimental/GameplayTargetingSystem/Source/GameplayTargetingSystem/Private/Tasks/TargetingSortTask_Base.cpp b/Engine/Plugins/Experimental/GameplayTargetingSystem/Source/GameplayTargetingSystem/Private/Tasks/TargetingSortTask_Base.cpp
--- a/Engine/Plugins/Experimental/GameplayTargetingSystem/Source/GameplayTargetingSystem/Private/Tasks/TargetingSortTask_Base.cpp
+++ b/Engine/Plugins/Experimental/GameplayTargetingSystem/Source/GameplayTargetingSystem/Private/Tasks/TargetingSortTask_Base.cpp
@@ -20,6 +20,7 @@ UTargetingSortTask_Base::UTargetingSortTask_Base(const FObjectInitializer& Objec
 	: Super(ObjectInitializer)
 {
 	bAscending = true;
+	bStableSort = false;
 }
 
 float UTargetingSortTask_Base::GetScoreForTarget(const FTargetingRequestHandle& TargetingHandle, const FTargetingDefaultResultData& TargetData) const
@@ -71,11 +72,20 @@ void UTargetingSortTask_Base::Execute(const FTargetingRequestHandle& TargetingHa
 				}
 			}
 
+			auto ByScore = [](const FTargetingDefaultResultData& Lhs, const FTargetingDefaultResultData& Rhs)
+			{
+				return Lhs.Score < Rhs.Score;
+			};
+
 			// sort the set
-			ResultData->TargetResults.Sort([this](const FTargetingDefaultResultData& Lhs, const FTargetingDefaultResultData& Rhs)
-				{
-					return Lhs.Score < Rhs.Score;
-				});
+			if (bStableSort)
+			{
+				ResultData->TargetResults.StableSort(ByScore);
+			}
+			else
+			{
+				ResultData->TargetResults.Sort(ByScore);
+			}
 
 #if ENABLE_DRAW_DEBUG
 			BuildPostSortDebugString(TargetingHandle, ResultData->TargetResults);
diff --git a/Engine/Plugins/Experimental/GameplayTargetingSystem/Source/GameplayTargetingSystem/Public/Tasks/TargetingSortTask_Base.h b/Engine/Plugins/Experimental/GameplayTargetingSystem/Source/GameplayTargetingSystem/Public/Tasks/TargetingSortTask_Base.h
--- a/Engine/Plugins/Experimental/GameplayTargetingSystem/Source/GameplayTargetingSystem/Public/Tasks/TargetingSortTask_Base.h
+++ b/Engine/Plugins/Experimental/GameplayTargetingSystem/Source/GameplayTargetingSystem/Public/Tasks/TargetingSortTask_Base.h
@@ -31,6 +31,10 @@ protected:
 	UPROPERTY(EditAnywhere, Category = "Target Sorting | Data")
 	uint8 bAscending : 1;
 
+	/** Should this task use a (slightly slower) sorting algorithm that preserves the relative ordering of targets with equal scores? */
+	UPROPERTY(EditAnywhere, Category = "Target Sorting | Data")
+	uint8 bStableSort : 1;
+
 	/** Debug Helper Methods */
 #if ENABLE_DRAW_DEBUG
 private:
