Fix null pointer dereference in USimpleTargetingSelectionTask::AddHitResult
#rnx

diff --git a/Engine/Plugins/Experimental/GameplayTargetingSystem/Source/GameplayTargetingSystem/Private/Tasks/SimpleTargetingSelectionTask.cpp b/Engine/Plugins/Experimental/GameplayTargetingSystem/Source/GameplayTargetingSystem/Private/Tasks/SimpleTargetingSelectionTask.cpp
--- a/Engine/Plugins/Experimental/GameplayTargetingSystem/Source/GameplayTargetingSystem/Private/Tasks/SimpleTargetingSelectionTask.cpp
+++ b/Engine/Plugins/Experimental/GameplayTargetingSystem/Source/GameplayTargetingSystem/Private/Tasks/SimpleTargetingSelectionTask.cpp
@@ -54,7 +54,7 @@ bool USimpleTargetingSelectionTask::AddHitResult(const FTargetingRequestHandle&
 		ResultData.HitResult = HitResult;
 		if (const FTargetingSourceContext* SourceContext = FTargetingSourceContext::Find(TargetingHandle))
 		{
-			if (SourceContext->SourceActor)
+			if (SourceContext->SourceActor && HitResult.HasValidHitObjectHandle())
 			{
 				ResultData.HitResult.Distance = FVector::Distance(SourceContext->SourceActor->GetActorLocation(), HitResult.GetActor()->GetActorLocation());
 			}
