Trace Targeting Selection / Engineering: Added option to use Single trace instead of Multi trace.
This allows for ensuring only one hit result is given, avoiding the case where an initial overlap hit is given along with the first hit past the initial overlap.

[RN]

diff --git a/Engine/Plugins/Experimental/GameplayTargetingSystem/Source/GameplayTargetingSystem/Private/Tasks/TargetingSelectionTask_Trace.cpp b/Engine/Plugins/Experimental/GameplayTargetingSystem/Source/GameplayTargetingSystem/Private/Tasks/TargetingSelectionTask_Trace.cpp
--- a/Engine/Plugins/Experimental/GameplayTargetingSystem/Source/GameplayTargetingSystem/Private/Tasks/TargetingSelectionTask_Trace.cpp
+++ b/Engine/Plugins/Experimental/GameplayTargetingSystem/Source/GameplayTargetingSystem/Private/Tasks/TargetingSelectionTask_Trace.cpp
@@ -19,6 +19,7 @@
 UTargetingSelectionTask_Trace::UTargetingSelectionTask_Trace(const FObjectInitializer& ObjectInitializer)
 : Super(ObjectInitializer)
 {
+	bMultiTrace = true;
 	bComplexTrace = false;
 	bIgnoreSourceActor = false;
 	bIgnoreInstigatorActor = false;
@@ -143,51 +144,81 @@ void UTargetingSelectionTask_Trace::ExecuteImmediateTrace(const FTargetingReques
 		FCollisionQueryParams Params(SCENE_QUERY_STAT(ExecuteImmediateTrace), bComplexTrace);
 		InitCollisionParams(TargetingHandle, Params);
 
+		FCollisionShape CollisionShape;
+		switch (TraceType)
+		{
+		case ETargetingTraceType::Sphere:
+			CollisionShape = FCollisionShape::MakeSphere(GetSweptTraceRadius(TargetingHandle));
+			break;
+		case ETargetingTraceType::Capsule:
+		{
+			const FVector CapsuleShapeVector = FVector(0.0f, GetSweptTraceRadius(TargetingHandle), GetSweptTraceCapsuleHalfHeight(TargetingHandle));
+			CollisionShape = FCollisionShape::MakeCapsule(CapsuleShapeVector);
+		}
+		break;
+		case ETargetingTraceType::Box:
+			CollisionShape = FCollisionShape::MakeBox(GetSweptTraceBoxHalfExtents(TargetingHandle));
+			break;
+		}
+
 		bool bHasBlockingHit = false;
 		TArray<FHitResult> Hits;
+
+		if (!bMultiTrace)
+		{
+			// Populate an element in the array for single trace to use
+			Hits.SetNum(1);
+		}
+
 		if (CollisionProfileName.Name != TEXT("NoCollision"))
 		{
-			switch (TraceType)
-			{
-			case ETargetingTraceType::Sphere:
-				bHasBlockingHit = World->SweepMultiByProfile(Hits, Start, End, FQuat::Identity, CollisionProfileName.Name, FCollisionShape::MakeSphere(GetSweptTraceRadius(TargetingHandle)), Params);
-				break;
-			case ETargetingTraceType::Capsule:
+			if (TraceType == ETargetingTraceType::Line)
 			{
-				const FVector CapsuleShapeVector = FVector(0.0f, GetSweptTraceRadius(TargetingHandle), GetSweptTraceCapsuleHalfHeight(TargetingHandle));
-				bHasBlockingHit = World->SweepMultiByProfile(Hits, Start, End, OrientationQuat, CollisionProfileName.Name, FCollisionShape::MakeCapsule(CapsuleShapeVector), Params);
+				if (bMultiTrace)
+				{
+					bHasBlockingHit = World->LineTraceMultiByProfile(Hits, Start, End, CollisionProfileName.Name, Params);
+				}
+				else
+				{
+					bHasBlockingHit = World->LineTraceSingleByProfile(Hits[0], Start, End, CollisionProfileName.Name, Params);
+				}
 			}
-				break;
-			case ETargetingTraceType::Box:
-				bHasBlockingHit = World->SweepMultiByProfile(Hits, Start, End, OrientationQuat, CollisionProfileName.Name, FCollisionShape::MakeBox(GetSweptTraceBoxHalfExtents(TargetingHandle)), Params);
-				break;
-			default:
-			case ETargetingTraceType::Line:
-				bHasBlockingHit = World->LineTraceMultiByProfile(Hits, Start, End, CollisionProfileName.Name, Params);
-				break;
+			else
+			{
+				if (bMultiTrace)
+				{
+					bHasBlockingHit = World->SweepMultiByProfile(Hits, Start, End, OrientationQuat, CollisionProfileName.Name, CollisionShape, Params);
+				}
+				else
+				{
+					bHasBlockingHit = World->SweepSingleByProfile(Hits[0], Start, End, OrientationQuat, CollisionProfileName.Name, CollisionShape, Params);
+				}
 			}
 		}
 		else
 		{
 			const ECollisionChannel CollisionChannel = UEngineTypes::ConvertToCollisionChannel(TraceChannel);
-			switch (TraceType)
+			if (TraceType == ETargetingTraceType::Line)
 			{
-			case ETargetingTraceType::Sphere:
-				bHasBlockingHit = World->SweepMultiByChannel(Hits, Start, End, FQuat::Identity, CollisionChannel, FCollisionShape::MakeSphere(GetSweptTraceRadius(TargetingHandle)), Params);
-				break;
-			case ETargetingTraceType::Capsule:
-			{
-				const FVector CapsuleShapeVector = FVector(0.0f, GetSweptTraceRadius(TargetingHandle), GetSweptTraceCapsuleHalfHeight(TargetingHandle));
-				bHasBlockingHit = World->SweepMultiByChannel(Hits, Start, End, OrientationQuat, CollisionChannel, FCollisionShape::MakeCapsule(CapsuleShapeVector), Params);
+				if (bMultiTrace)
+				{
+					bHasBlockingHit = World->LineTraceMultiByChannel(Hits, Start, End, CollisionChannel, Params);
+				}
+				else
+				{
+					bHasBlockingHit = World->LineTraceSingleByChannel(Hits[0], Start, End, CollisionChannel, Params);
+				}
 			}
-				break;
-			case ETargetingTraceType::Box:
-				bHasBlockingHit = World->SweepMultiByChannel(Hits, Start, End, OrientationQuat, CollisionChannel, FCollisionShape::MakeBox(GetSweptTraceBoxHalfExtents(TargetingHandle)), Params);
-				break;
-			default:
-			case ETargetingTraceType::Line:
-				bHasBlockingHit = World->LineTraceMultiByChannel(Hits, Start, End, CollisionChannel, Params);
-				break;
+			else
+			{
+				if (bMultiTrace)
+				{
+					bHasBlockingHit = World->SweepMultiByChannel(Hits, Start, End, OrientationQuat, CollisionChannel, CollisionShape, Params);
+				}
+				else
+				{
+					bHasBlockingHit = World->SweepSingleByChannel(Hits[0], Start, End, OrientationQuat, CollisionChannel, CollisionShape, Params);
+				}
 			}
 		}
 
diff --git a/Engine/Plugins/Experimental/GameplayTargetingSystem/Source/GameplayTargetingSystem/Public/Tasks/TargetingSelectionTask_Trace.h b/Engine/Plugins/Experimental/GameplayTargetingSystem/Source/GameplayTargetingSystem/Public/Tasks/TargetingSelectionTask_Trace.h
--- a/Engine/Plugins/Experimental/GameplayTargetingSystem/Source/GameplayTargetingSystem/Public/Tasks/TargetingSelectionTask_Trace.h
+++ b/Engine/Plugins/Experimental/GameplayTargetingSystem/Source/GameplayTargetingSystem/Public/Tasks/TargetingSelectionTask_Trace.h
@@ -158,6 +158,10 @@ protected:
 	UPROPERTY(EditAnywhere, Category = "Target Trace Selection | Trace Data")
 	uint8 bComplexTrace : 1;
 
+	/** Indicates whether the trace should be a multi trace or a single trace */
+	UPROPERTY(EditAnywhere, Category = "Target Trace Selection | Trace Data")
+	uint8 bMultiTrace : 1;
+
 	/** Indicates the trace should ignore the source actor */
 	UPROPERTY(EditAnywhere, Category = "Target Trace Selection | Trace Data")
 	uint8 bIgnoreSourceActor : 1;
