!B - Targeting selection debug for cylinders display correctly and enforcing rules to avoid incorrect overlap detection for cylinders.

diff --git a/Engine/Plugins/Experimental/GameplayTargetingSystem/Source/GameplayTargetingSystem/Private/Tasks/TargetingSelectionTask_AOE.cpp b/Engine/Plugins/Experimental/GameplayTargetingSystem/Source/GameplayTargetingSystem/Private/Tasks/TargetingSelectionTask_AOE.cpp
--- a/Engine/Plugins/Experimental/GameplayTargetingSystem/Source/GameplayTargetingSystem/Private/Tasks/TargetingSelectionTask_AOE.cpp
+++ b/Engine/Plugins/Experimental/GameplayTargetingSystem/Source/GameplayTargetingSystem/Private/Tasks/TargetingSelectionTask_AOE.cpp
@@ -149,7 +149,7 @@ void UTargetingSelectionTask_AOE::DebugDrawBoundingVolume(const FTargetingReques
 		break;
 	case ETargetingAOEShape::Cylinder:
 		{
-			const FVector RotatedExtent = SourceRotation * CollisionShape.GetExtent();
+			const FVector RotatedExtent = SourceRotation * FVector(0, 0, CollisionShape.GetExtent().Z);
 			DrawDebugCylinder(World, SourceLocation - RotatedExtent, SourceLocation + RotatedExtent, CollisionShape.GetExtent().X, 32,
 				Color, bPersistentLines, LifeTime, DepthPriority, Thickness);
 			break;
@@ -443,6 +443,24 @@ bool UTargetingSelectionTask_AOE::CanEditChange(const FProperty* InProperty) con
 
 	return true;
 }
+
+void UTargetingSelectionTask_AOE::PostEditChangeProperty(struct FPropertyChangedEvent& PropertyChangedEvent)
+{
+	if (PropertyChangedEvent.MemberProperty)
+	{
+		if (PropertyChangedEvent.MemberProperty->GetFName() == GET_MEMBER_NAME_CHECKED(UTargetingSelectionTask_AOE, HalfExtent) ||
+			PropertyChangedEvent.MemberProperty->GetFName() == GET_MEMBER_NAME_CHECKED(UTargetingSelectionTask_AOE, ShapeType))
+		{
+			// If using Cylinder make sure the X and Y match so that we don't get incorrect Overlaps when Y < X
+			if (ShapeType == ETargetingAOEShape::Cylinder)
+			{
+				HalfExtent.Y = HalfExtent.X;
+			}
+		}
+	}
+
+	Super::PostEditChangeProperty(PropertyChangedEvent);
+}
 #endif // WITH_EDITOR
 
 #if ENABLE_DRAW_DEBUG
diff --git a/Engine/Plugins/Experimental/GameplayTargetingSystem/Source/GameplayTargetingSystem/Public/Tasks/TargetingSelectionTask_AOE.h b/Engine/Plugins/Experimental/GameplayTargetingSystem/Source/GameplayTargetingSystem/Public/Tasks/TargetingSelectionTask_AOE.h
--- a/Engine/Plugins/Experimental/GameplayTargetingSystem/Source/GameplayTargetingSystem/Public/Tasks/TargetingSelectionTask_AOE.h
+++ b/Engine/Plugins/Experimental/GameplayTargetingSystem/Source/GameplayTargetingSystem/Public/Tasks/TargetingSelectionTask_AOE.h
@@ -135,7 +135,9 @@ protected:
 	UPROPERTY(EditAnywhere, Category = "Target AOE Selection | Data")
 	uint8 bTraceComplex : 1 = false;
 
-	/** The half extent to use for box and cylinder */
+	/** The half extent to use for box and cylinder
+	 *	For Cylinder the X component will be the radius and Z will be its height
+	 */
 	UPROPERTY(EditAnywhere, Category = "Target AOE Selection | Box and Cylinder")
 	FVector HalfExtent = FVector::ZeroVector;
 
@@ -155,6 +157,7 @@ protected:
 protected:
 #if WITH_EDITOR
 	virtual bool CanEditChange(const FProperty* InProperty) const override;
+	virtual void PostEditChangeProperty(struct FPropertyChangedEvent& PropertyChangedEvent) override;
 #endif
 
 	/** Debug Helper Methods */
