Allow rotating AOE target selection shapes, so that a capsule can be used horizontally instead of vertically for example

#rnx

diff --git a/Engine/Plugins/Experimental/GameplayTargetingSystem/Source/GameplayTargetingSystem/Private/Tasks/TargetingSelectionTask_AOE.cpp b/Engine/Plugins/Experimental/GameplayTargetingSystem/Source/GameplayTargetingSystem/Private/Tasks/TargetingSelectionTask_AOE.cpp
--- a/Engine/Plugins/Experimental/GameplayTargetingSystem/Source/GameplayTargetingSystem/Private/Tasks/TargetingSelectionTask_AOE.cpp
+++ b/Engine/Plugins/Experimental/GameplayTargetingSystem/Source/GameplayTargetingSystem/Private/Tasks/TargetingSelectionTask_AOE.cpp
@@ -64,7 +64,7 @@ void UTargetingSelectionTask_AOE::ExecuteImmediateTrace(const FTargetingRequestH
 	if (World && TargetingHandle.IsValid())
 	{
 		const FVector SourceLocation = GetSourceLocation(TargetingHandle) + GetSourceOffset(TargetingHandle);
-		const FQuat SourceRotation = GetSourceRotation(TargetingHandle);
+		const FQuat SourceRotation = GetSourceRotation(TargetingHandle) * GetSourceRotationOffset(TargetingHandle).Quaternion();
 
 		TArray<FOverlapResult> OverlapResults;
 		if (ShapeType == ETargetingAOEShape::SourceComponent)
@@ -125,7 +125,7 @@ void UTargetingSelectionTask_AOE::DebugDrawBoundingVolume(const FTargetingReques
 #if ENABLE_DRAW_DEBUG
 	const UWorld* World = GetSourceContextWorld(TargetingHandle);
 	const FVector SourceLocation = OverlapDatum ? OverlapDatum->Pos : GetSourceLocation(TargetingHandle) + GetSourceOffset(TargetingHandle);
-	const FQuat SourceRotation = OverlapDatum ? OverlapDatum->Rot : GetSourceRotation(TargetingHandle);
+	const FQuat SourceRotation = OverlapDatum ? OverlapDatum->Rot : GetSourceRotation(TargetingHandle) * GetSourceRotationOffset(TargetingHandle).Quaternion();
 	const FCollisionShape CollisionShape = GetCollisionShape();
 
 	const bool bPersistentLines = false;
@@ -184,7 +184,7 @@ void UTargetingSelectionTask_AOE::ExecuteAsyncTrace(const FTargetingRequestHandl
 	if (World && TargetingHandle.IsValid())
 	{
 		const FVector SourceLocation = GetSourceLocation(TargetingHandle) + GetSourceOffset(TargetingHandle);
-		const FQuat SourceRotation = GetSourceRotation(TargetingHandle);
+		const FQuat SourceRotation = GetSourceRotation(TargetingHandle) * GetSourceRotationOffset(TargetingHandle).Quaternion();
 
 		FCollisionShape CollisionShape = GetCollisionShape();
 		FCollisionQueryParams OverlapParams(TEXT("UTargetingSelectionTask_AOE"), SCENE_QUERY_STAT_ONLY(UTargetingSelectionTask_AOE_Shape), false);
@@ -337,6 +337,11 @@ FQuat UTargetingSelectionTask_AOE::GetSourceRotation_Implementation(const FTarge
 	return FQuat::Identity;
 }
 
+FRotator UTargetingSelectionTask_AOE::GetSourceRotationOffset_Implementation(const FTargetingRequestHandle& TargetingHandle) const
+{
+	return DefaultSourceRotationOffset;
+}
+
 FCollisionShape UTargetingSelectionTask_AOE::GetCollisionShape() const
 {
 	switch (ShapeType)
diff --git a/Engine/Plugins/Experimental/GameplayTargetingSystem/Source/GameplayTargetingSystem/Public/Tasks/TargetingSelectionTask_AOE.h b/Engine/Plugins/Experimental/GameplayTargetingSystem/Source/GameplayTargetingSystem/Public/Tasks/TargetingSelectionTask_AOE.h
--- a/Engine/Plugins/Experimental/GameplayTargetingSystem/Source/GameplayTargetingSystem/Public/Tasks/TargetingSelectionTask_AOE.h
+++ b/Engine/Plugins/Experimental/GameplayTargetingSystem/Source/GameplayTargetingSystem/Public/Tasks/TargetingSelectionTask_AOE.h
@@ -89,6 +89,10 @@ protected:
 	/** Native event to get the source rotation for the AOE  */
 	UFUNCTION(BlueprintNativeEvent, Category = "Target AOE Selection")
 	UE_API FQuat GetSourceRotation(const FTargetingRequestHandle& TargetingHandle) const;
+
+	/** Native Event to get a source rotation offset for the AOE */
+	UFUNCTION(BlueprintNativeEvent, Category = "Target AOE Selection")
+	UE_API FRotator GetSourceRotationOffset(const FTargetingRequestHandle& TargetingHandle) const;
 	
 private:
 	/** Helper method to build the Collision Shape */
@@ -121,6 +125,10 @@ protected:
 	UPROPERTY(EditAnywhere, Category = "Target AOE Selection | Data")
 	FVector DefaultSourceOffset = FVector::ZeroVector;
 
+	/** The default source rotation offset used by GetSourceRotationOffset */
+	UPROPERTY(EditAnywhere, Category = "Target AOE Selection | Data")
+	FRotator DefaultSourceRotationOffset = FRotator::ZeroRotator;
+
 	/** Should we offset based on world or relative Source object transform? */
 	UPROPERTY(EditAnywhere, Category = "Target AOE Selection | Data")
 	uint8 bUseRelativeOffset : 1;
