Expose some targeting system enum, functions and type

#rnx
#Synced-CL 49811063

diff --git a/Engine/Plugins/Experimental/GameplayTargetingSystem/Source/GameplayTargetingSystem/Public/Tasks/TargetingSelectionTask_AOE.h b/Engine/Plugins/Experimental/GameplayTargetingSystem/Source/GameplayTargetingSystem/Public/Tasks/TargetingSelectionTask_AOE.h
--- a/Engine/Plugins/Experimental/GameplayTargetingSystem/Source/GameplayTargetingSystem/Public/Tasks/TargetingSelectionTask_AOE.h
+++ b/Engine/Plugins/Experimental/GameplayTargetingSystem/Source/GameplayTargetingSystem/Public/Tasks/TargetingSelectionTask_AOE.h
@@ -23,7 +23,7 @@ struct FTargetingRequestHandle;
 /**
 *	@enum ETargetingAOEShape	
 */
-UENUM()
+UENUM(BlueprintType)
 enum class ETargetingAOEShape : uint8
 {
 	Box,
@@ -59,25 +59,12 @@ public:
 	/** Debug draws the outlines of the set shape type. */
 	UE_API void DebugDrawBoundingVolume(const FTargetingRequestHandle& TargetingHandle, const FColor& Color, const FOverlapDatum* OverlapDatum = nullptr) const;
 
+	UE_API ETargetingAOEShape GetShapeType() const { return ShapeType; }
 	UE_API void SetShapeType(ETargetingAOEShape InShapeType);
 	UE_API void SetHalfExtent(FVector InHalfExtent);
 	UE_API void SetRadius(FScalableFloat InRadius);
 	UE_API void SetHalfHeight(FScalableFloat InHalfHeight);
 
-private:
-	/** Method to process the trace task immediately */
-	void ExecuteImmediateTrace(const FTargetingRequestHandle& TargetingHandle) const;
-
-	/** Method to process the trace task asynchronously */
-	void ExecuteAsyncTrace(const FTargetingRequestHandle& TargetingHandle) const;
-
-	/** Callback for an async overlap */
-	void HandleAsyncOverlapComplete(const FTraceHandle& InTraceHandle, FOverlapDatum& InOverlapDatum, FTargetingRequestHandle TargetingHandle) const;
-
-	/** Method to take the overlap results and store them in the targeting result data */
-	void ProcessOverlapResults(const FTargetingRequestHandle& TargetingHandle, const TArray<FOverlapResult>& Overlaps) const;
-
-protected:
 	/** Native Event to get the source location for the AOE */
 	UFUNCTION(BlueprintNativeEvent, Category = "Target AOE Selection")
 	UE_API FVector GetSourceLocation(const FTargetingRequestHandle& TargetingHandle) const;
@@ -93,11 +80,26 @@ protected:
 	/** Native Event to get a source rotation offset for the AOE */
 	UFUNCTION(BlueprintNativeEvent, Category = "Target AOE Selection")
 	UE_API FRotator GetSourceRotationOffset(const FTargetingRequestHandle& TargetingHandle) const;
-	
-private:
+
 	/** Helper method to build the Collision Shape */
-	FCollisionShape GetCollisionShape() const;
+	UE_API FCollisionShape GetCollisionShape() const;
+
+	UE_API bool GetUseRelativeOffset() const { return bUseRelativeOffset; }
+
+private:
+	/** Method to process the trace task immediately */
+	void ExecuteImmediateTrace(const FTargetingRequestHandle& TargetingHandle) const;
+
+	/** Method to process the trace task asynchronously */
+	void ExecuteAsyncTrace(const FTargetingRequestHandle& TargetingHandle) const;
+
+	/** Callback for an async overlap */
+	void HandleAsyncOverlapComplete(const FTraceHandle& InTraceHandle, FOverlapDatum& InOverlapDatum, FTargetingRequestHandle TargetingHandle) const;
 
+	/** Method to take the overlap results and store them in the targeting result data */
+	void ProcessOverlapResults(const FTargetingRequestHandle& TargetingHandle, const TArray<FOverlapResult>& Overlaps) const;
+	
+private:
 	/** Helper method to find the custom component defined on the source actor */
 	const UPrimitiveComponent* GetCollisionComponent(const FTargetingRequestHandle& TargetingHandle) const;
 
diff --git a/Engine/Plugins/Experimental/GameplayTargetingSystem/Source/GameplayTargetingSystem/Public/Types/TargetingSystemTypes.h b/Engine/Plugins/Experimental/GameplayTargetingSystem/Source/GameplayTargetingSystem/Public/Types/TargetingSystemTypes.h
--- a/Engine/Plugins/Experimental/GameplayTargetingSystem/Source/GameplayTargetingSystem/Public/Types/TargetingSystemTypes.h
+++ b/Engine/Plugins/Experimental/GameplayTargetingSystem/Source/GameplayTargetingSystem/Public/Types/TargetingSystemTypes.h
@@ -80,7 +80,7 @@ struct TStructOpsTypeTraits<FTargetingRequestHandle> : public TStructOpsTypeTrai
 *
 *	A set of tasks to be used by targeting requests to find/processes targets
 */
-USTRUCT(BlueprintType)
+USTRUCT(MinimalAPI, BlueprintType)
 struct FTargetingTaskSet
 {
 	GENERATED_BODY()
